MeshModelingTools: UVEditor compilation fixes.

#rnx

diff --git a/Engine/Plugins/Experimental/UVEditor/Source/UVEditor/Private/ToolTargets/StaticMeshUVMeshToolTarget.cpp b/Engine/Plugins/Experimental/UVEditor/Source/UVEditor/Private/ToolTargets/StaticMeshUVMeshToolTarget.cpp
--- a/Engine/Plugins/Experimental/UVEditor/Source/UVEditor/Private/ToolTargets/StaticMeshUVMeshToolTarget.cpp
+++ b/Engine/Plugins/Experimental/UVEditor/Source/UVEditor/Private/ToolTargets/StaticMeshUVMeshToolTarget.cpp
@@ -20,6 +20,11 @@ namespace StaticMeshUVMeshToolTargetLocals
 	const double UVScalingFactor = 1000;
 }
 
+bool UStaticMeshUVMeshToolTarget::IsValid() const
+{
+	return OriginalAsset && !OriginalAsset->IsPendingKillOrUnreachable() && OriginalAsset->IsValidLowLevel();
+}
+
 int32 UStaticMeshUVMeshToolTarget::GetNumUVLayers()
 {
 	using namespace StaticMeshUVMeshToolTargetLocals;
diff --git a/Engine/Plugins/Experimental/UVEditor/Source/UVEditor/Private/UVEditorModeToolkit.cpp b/Engine/Plugins/Experimental/UVEditor/Source/UVEditor/Private/UVEditorModeToolkit.cpp
--- a/Engine/Plugins/Experimental/UVEditor/Source/UVEditor/Private/UVEditorModeToolkit.cpp
+++ b/Engine/Plugins/Experimental/UVEditor/Source/UVEditor/Private/UVEditorModeToolkit.cpp
@@ -2,8 +2,15 @@
 
 #include "UVEditorModeToolkit.h"
 
-#include "UVEditorCommands.h"
+#include "EditorStyleSet.h" //FEditorStyle
+#include "Framework/MultiBox/MultiBoxBuilder.h"
+#include "Framework/MultiBox/MultiBoxExtender.h"
+#include "Framework/MultiBox/MultiBoxDefs.h"
+#include "IDetailsView.h"
 #include "Tools/UEdMode.h"
+#include "UVEditorCommands.h"
+#include "Widgets/Layout/SBorder.h"
+#include "Widgets/Text/STextBlock.h"
 
 #define LOCTEXT_NAMESPACE "FUVEditorModeToolkit"
 
diff --git a/Engine/Plugins/Experimental/UVEditor/Source/UVEditor/Private/UVEditorToolkit.cpp b/Engine/Plugins/Experimental/UVEditor/Source/UVEditor/Private/UVEditorToolkit.cpp
--- a/Engine/Plugins/Experimental/UVEditor/Source/UVEditor/Private/UVEditorToolkit.cpp
+++ b/Engine/Plugins/Experimental/UVEditor/Source/UVEditor/Private/UVEditorToolkit.cpp
@@ -2,8 +2,9 @@
 
 #include "UVEditorToolkit.h"
 
-#include "PreviewScene.h"
 #include "AssetEditorModeManager.h"
+#include "Misc/MessageDialog.h"
+#include "PreviewScene.h"
 #include "SAssetEditorViewport.h"
 #include "ToolMenus.h"
 #include "UVEditor.h"
diff --git a/Engine/Plugins/Experimental/UVEditor/Source/UVEditor/Public/ToolTargets/StaticMeshUVMeshToolTarget.h b/Engine/Plugins/Experimental/UVEditor/Source/UVEditor/Public/ToolTargets/StaticMeshUVMeshToolTarget.h
--- a/Engine/Plugins/Experimental/UVEditor/Source/UVEditor/Public/ToolTargets/StaticMeshUVMeshToolTarget.h
+++ b/Engine/Plugins/Experimental/UVEditor/Source/UVEditor/Public/ToolTargets/StaticMeshUVMeshToolTarget.h
@@ -34,6 +34,9 @@ public:
 	virtual TSharedPtr<FDynamicMesh3> GetMesh(int32 LayerIndex) override;
 	virtual void SaveBackToUVs(const FDynamicMesh3* MeshToSave, int32 LayerIndex) override;
 
+	// UToolTarget
+	virtual bool IsValid() const override;
+
 protected:
 	UPROPERTY()
 	UStaticMesh* OriginalAsset = nullptr;
diff --git a/Engine/Plugins/Experimental/UVEditor/Source/UVEditor/Public/UVEditorModeToolkit.h b/Engine/Plugins/Experimental/UVEditor/Source/UVEditor/Public/UVEditorModeToolkit.h
--- a/Engine/Plugins/Experimental/UVEditor/Source/UVEditor/Public/UVEditorModeToolkit.h
+++ b/Engine/Plugins/Experimental/UVEditor/Source/UVEditor/Public/UVEditorModeToolkit.h
@@ -6,6 +6,9 @@
 
 #include "Toolkits/BaseToolkit.h"
 
+class SBorder;
+class STextBlock;
+
 /**
  * The UV editor mode toolkit is responsible for the panel on the side in the UV editor
  * that shows mode and tool properties. Tool buttons would go in Init().
diff --git a/Engine/Plugins/Experimental/UVEditor/Source/UVEditorTools/Public/UVEditorToolsModule.h b/Engine/Plugins/Experimental/UVEditor/Source/UVEditorTools/Public/UVEditorToolsModule.h
--- a/Engine/Plugins/Experimental/UVEditor/Source/UVEditorTools/Public/UVEditorToolsModule.h
+++ b/Engine/Plugins/Experimental/UVEditor/Source/UVEditorTools/Public/UVEditorToolsModule.h
@@ -3,7 +3,7 @@
 #pragma once
 
 #include "CoreMinimal.h"
-#include "Modules/ModuleInterface.h"
+#include "Modules/ModuleManager.h"
 
 class FUVEditorToolsModule : public IModuleInterface
 {
