MeshModelingTools: Fix UV editor check hit during tooltip because asset toolkit expects only one edited asset.

#rnx

diff --git a/Engine/Plugins/Experimental/UVEditor/Source/UVEditor/Private/UVEditorToolkit.cpp b/Engine/Plugins/Experimental/UVEditor/Source/UVEditor/Private/UVEditorToolkit.cpp
--- a/Engine/Plugins/Experimental/UVEditor/Source/UVEditor/Private/UVEditorToolkit.cpp
+++ b/Engine/Plugins/Experimental/UVEditor/Source/UVEditor/Private/UVEditorToolkit.cpp
@@ -80,10 +80,10 @@ FText FUVEditorToolkit::GetToolkitName() const
 	const TArray<UObject*>* Objects = GetObjectsCurrentlyBeingEdited();
 	if (Objects->Num() == 1)
 	{
-		return FText::Format(LOCTEXT("UVEditorTabNameWithObject", "{0} UVs"), 
+		return FText::Format(LOCTEXT("UVEditorTabNameWithObject", "UVs: {0}"), 
 			GetLabelForObject((*Objects)[0]));
 	}
-	return LOCTEXT("UVEditorSimpleTabName", "UV Editor");
+	return LOCTEXT("UVEditorMultipleTabName", "UVs: Multiple");
 }
 
 // This gets used multiple places, most notably in GetToolMenuAppName, which gets
@@ -99,6 +99,24 @@ FText FUVEditorToolkit::GetBaseToolkitName() const
 	return LOCTEXT("UVBaseToolkitName", "UV");
 }
 
+FText FUVEditorToolkit::GetToolkitToolTipText() const
+{
+	FString ToolTipString;
+	ToolTipString += LOCTEXT("ToolTipAssetLabel", "Asset").ToString();
+	ToolTipString += TEXT(": ");
+
+	const TArray<UObject*>* Objects = GetObjectsCurrentlyBeingEdited();
+	check(Objects && Objects->Num() > 0);
+	ToolTipString += GetLabelForObject((*Objects)[0]).ToString();
+	for (int32 i = 1; i < Objects->Num(); ++i)
+	{
+		ToolTipString += TEXT(", ");
+		ToolTipString += GetLabelForObject((*Objects)[i]).ToString();
+	}
+
+	return FText::FromString(ToolTipString);
+}
+
 void FUVEditorToolkit::RegisterTabSpawners(const TSharedRef<FTabManager>& InTabManager)
 {
 	// Here we set up the tabs we referenced in StandaloneDefaultLayout (in the constructor).
diff --git a/Engine/Plugins/Experimental/UVEditor/Source/UVEditor/Public/UVEditorToolkit.h b/Engine/Plugins/Experimental/UVEditor/Source/UVEditor/Public/UVEditorToolkit.h
--- a/Engine/Plugins/Experimental/UVEditor/Source/UVEditor/Public/UVEditorToolkit.h
+++ b/Engine/Plugins/Experimental/UVEditor/Source/UVEditor/Public/UVEditorToolkit.h
@@ -38,6 +38,7 @@ public:
 	virtual FText GetToolkitName() const override;
 	virtual FName GetToolkitFName() const override;
 	virtual FText GetBaseToolkitName() const override;
+	virtual FText GetToolkitToolTipText() const override;
 	virtual void RegisterTabSpawners(const TSharedRef<FTabManager>& InTabManager) override;
 	virtual bool OnRequestClose() override;
 	virtual void SaveAsset_Execute() override;
