UVEditor: Added live preview to UV editor. To support this, did major refactoring in the way that the UV editor supplies its tools with input. Also added scaling of islands and tweaked the used material to make gizmo more visible.

#rnx

diff --git a/Engine/Plugins/Experimental/MeshModelingToolset/Source/ModelingComponents/Private/Selection/MeshSelectionMechanic.cpp b/Engine/Plugins/Experimental/MeshModelingToolset/Source/ModelingComponents/Private/Selection/MeshSelectionMechanic.cpp
--- a/Engine/Plugins/Experimental/MeshModelingToolset/Source/ModelingComponents/Private/Selection/MeshSelectionMechanic.cpp
+++ b/Engine/Plugins/Experimental/MeshModelingToolset/Source/ModelingComponents/Private/Selection/MeshSelectionMechanic.cpp
@@ -4,12 +4,15 @@
 
 #include "BaseBehaviors/BehaviorTargetInterfaces.h"
 #include "BaseBehaviors/SingleClickBehavior.h"
+#include "Drawing/LineSetComponent.h"
+#include "Drawing/PreviewGeometryActor.h"
 #include "InteractiveToolManager.h"
 #include "Polyline3.h"
 #include "Selections/MeshConnectedComponents.h"
 #include "SceneManagement.h"
 #include "Spatial/GeometrySet3.h"
 #include "ToolSceneQueriesUtil.h"
+#include "ToolSetupUtil.h"
 
 using namespace UE::Geometry;
 
@@ -23,16 +26,51 @@ void UMeshSelectionMechanic::Setup(UInteractiveTool* ParentToolIn)
 	USingleClickInputBehavior* ClickBehavior = NewObject<USingleClickInputBehavior>();
 	ClickBehavior->Initialize(this);
 	ParentTool->AddInputBehavior(ClickBehavior);
+
+	LineSet = NewObject<ULineSetComponent>();
+	LineSet->SetLineMaterial(ToolSetupUtil::GetDefaultLineComponentMaterial(
+		GetParentTool()->GetToolManager(), /*bDepthTested*/ true));
+}
+
+void UMeshSelectionMechanic::SetWorld(UWorld* World)
+{
+	// It may be unreasonable to worry about SetWorld being called more than once, but let's be safe anyway
+	if (PreviewGeometryActor)
+	{
+		PreviewGeometryActor->Destroy();
+	}
+
+	// We need the world so we can create the geometry actor in the right place.
+	FRotator Rotation(0.0f, 0.0f, 0.0f);
+	FActorSpawnParameters SpawnInfo;
+	PreviewGeometryActor = World->SpawnActor<APreviewGeometryActor>(FVector::ZeroVector, Rotation, SpawnInfo);
+
+	// Attach the rendering component to the actor
+	LineSet->Rename(nullptr, PreviewGeometryActor); // Changes the "outer"
+	PreviewGeometryActor->SetRootComponent(LineSet);
+	if (LineSet->IsRegistered())
+	{
+		LineSet->ReregisterComponent();
+	}
+	else
+	{
+		LineSet->RegisterComponent();
+	}
 }
 
 void UMeshSelectionMechanic::Shutdown()
 {
-	// TODO: Will need to destroy preview actors here once using those for rendering
+	if (PreviewGeometryActor)
+	{
+		PreviewGeometryActor->Destroy();
+		PreviewGeometryActor = nullptr;
+	}
 }
 
-void UMeshSelectionMechanic::AddSpatial(TSharedPtr<FDynamicMeshAABBTree3> SpatialIn)
+void UMeshSelectionMechanic::AddSpatial(TSharedPtr<FDynamicMeshAABBTree3> SpatialIn, const FTransform& TransformIn)
 {
 	MeshSpatials.Add(SpatialIn);
+	MeshTransforms.Add(TransformIn);
 }
 
 const FDynamicMeshSelection& UMeshSelectionMechanic::GetCurrentSelection() const
@@ -43,6 +81,22 @@ const FDynamicMeshSelection& UMeshSelectionMechanic::GetCurrentSelection() const
 void UMeshSelectionMechanic::SetSelection(const FDynamicMeshSelection& Selection, bool bBroadcast, bool bEmitChange)
 {
 	CurrentSelection = Selection;
+
+	if (MeshSpatials[CurrentSelectionIndex]->GetMesh() != Selection.Mesh)
+	{
+		for (int32 i = 0; i < MeshSpatials.Num(); ++i)
+		{
+			if (MeshSpatials[i]->GetMesh() == Selection.Mesh)
+			{
+				CurrentSelectionIndex = i;
+				break;
+			}
+		}
+	}
+
+	UpdateCentroid();
+	RebuildDrawnElements(FTransform(CurrentSelectionCentroid));
+
 	if (bBroadcast)
 	{
 		OnSelectionChanged.Broadcast();
@@ -50,32 +104,34 @@ void UMeshSelectionMechanic::SetSelection(const FDynamicMeshSelection& Selection
 	// TODO: Undo/redo
 }
 
-void UMeshSelectionMechanic::Render(IToolsContextRenderAPI* RenderAPI)
+void UMeshSelectionMechanic::RebuildDrawnElements(const FTransform& StartTransform)
 {
-	// TODO:We do this in other places, but should CameraState be cached somewhere else?
-	GetParentTool()->GetToolManager()->GetContextQueriesAPI()->GetCurrentViewState(CameraState);
+	LineSet->Clear();
+	PreviewGeometryActor->SetActorTransform(StartTransform);
 
-	if (CurrentSelection.IsEmpty())
+	// For us to end up with the StartTransform, we have to invert it, but only after applying
+	// the mesh transform to begin with.
+	auto TransformToApply = [&StartTransform, this](const FVector& VectorIn)
 	{
-		return;
-	}
-
-	// TODO: We should do all the selection with line/triangle components instead of PDI
-	FPrimitiveDrawInterface* PDI = RenderAPI->GetPrimitiveDrawInterface();
-	float PDIScale = RenderAPI->GetCameraState().GetPDIScalingFactor();
-	float LineThickness = 3.0 * PDIScale;
-	float DepthBias = 3.0f;
+		return StartTransform.InverseTransformPosition(
+			MeshTransforms[CurrentSelectionIndex].TransformPosition(VectorIn));
+	};
 
 	if (CurrentSelection.Type == FDynamicMeshSelection::EType::Triangle)
 	{
 		for (int32 Tid : CurrentSelection.SelectedIDs)
 		{
+			FIndex3i Vids = CurrentSelection.Mesh->GetTriangle(Tid);
+			FVector Points[3];
+			for (int i = 0; i < 3; ++i)
+			{
+				Points[i] = TransformToApply(CurrentSelection.Mesh->GetVertex(Vids[i]));
+			}
 			for (int i = 0; i < 3; ++i)
 			{
 				int NextIndex = (i + 1) % 3;
-				PDI->DrawLine((FVector)CurrentSelection.Mesh->GetTriVertex(Tid, i),
-					(FVector)CurrentSelection.Mesh->GetTriVertex(Tid, NextIndex), FLinearColor::Yellow,
-					SDPG_Foreground, LineThickness, DepthBias, true);
+				LineSet->AddLine(Points[i], Points[NextIndex],
+					LineColor, LineThickness, DepthBias);
 			}
 		}
 	}
@@ -84,11 +140,62 @@ void UMeshSelectionMechanic::Render(IToolsContextRenderAPI* RenderAPI)
 		for (int32 Eid : CurrentSelection.SelectedIDs)
 		{
 			FIndex2i EdgeVids = CurrentSelection.Mesh->GetEdgeV(Eid);
-			PDI->DrawLine((FVector)CurrentSelection.Mesh->GetVertex(EdgeVids.A),
-				(FVector)CurrentSelection.Mesh->GetVertex(EdgeVids.B), FLinearColor::Yellow,
-				SDPG_Foreground, LineThickness, DepthBias, true);
+			LineSet->AddLine(
+				TransformToApply(CurrentSelection.Mesh->GetVertex(EdgeVids.A)),
+				TransformToApply(CurrentSelection.Mesh->GetVertex(EdgeVids.B)),
+					LineColor, LineThickness, DepthBias);
+		}
+	}
+}
+
+void UMeshSelectionMechanic::UpdateCentroid()
+{
+	CurrentSelectionCentroid = FVector3d(0);
+	if (CurrentSelection.IsEmpty())
+	{
+		return;
+	}
+
+	if (CurrentSelection.Type == FDynamicMeshSelection::EType::Edge)
+	{
+		for (int32 Eid : CurrentSelection.SelectedIDs)
+		{
+			CurrentSelectionCentroid += CurrentSelection.Mesh->GetEdgePoint(Eid, 0.5);
 		}
+		CurrentSelectionCentroid /= CurrentSelection.SelectedIDs.Num();
+	}
+	else if (CurrentSelection.Type == FDynamicMeshSelection::EType::Triangle)
+	{
+		for (int32 Tid : CurrentSelection.SelectedIDs)
+		{
+			CurrentSelectionCentroid += CurrentSelection.Mesh->GetTriCentroid(Tid);
+		}
+		CurrentSelectionCentroid /= CurrentSelection.SelectedIDs.Num();
+	}
+}
+
+FVector3d UMeshSelectionMechanic::GetCurrentSelectionCentroid()
+{
+	return CurrentSelectionCentroid;
+}
+
+
+void UMeshSelectionMechanic::SetDrawnElementsTransform(const FTransform& Transform)
+{
+	PreviewGeometryActor->SetActorTransform(Transform);
+}
+
+void UMeshSelectionMechanic::Render(IToolsContextRenderAPI* RenderAPI)
+{
+	// TODO:We do this in other places, but should CameraState be cached somewhere else?
+	GetParentTool()->GetToolManager()->GetContextQueriesAPI()->GetCurrentViewState(CameraState);
+
+	if (CurrentSelection.IsEmpty())
+	{
+		return;
 	}
+
+
 }
 
 FInputRayHit UMeshSelectionMechanic::IsHitByClick(const FInputDeviceRay& ClickPos)
@@ -108,11 +215,18 @@ void UMeshSelectionMechanic::OnClicked(const FInputDeviceRay& ClickPos)
 
 	int32 HitTid = IndexConstants::InvalidID;
 	double RayT = 0; 
-	for (TSharedPtr<FDynamicMeshAABBTree3> Spatial : MeshSpatials)
+	for (int32 i = 0; i < MeshSpatials.Num(); ++i)
 	{
-		if (Spatial->FindNearestHitTriangle(ClickPos.WorldRay, RayT, HitTid))
+		// Minor note: Using FRay3d in this statement instead of FRay causes an internal compiler error in VS
+		// in Development Editor builds, for no good reason.
+		FRay LocalRay(
+			MeshTransforms[i].InverseTransformPosition(ClickPos.WorldRay.Origin),
+			MeshTransforms[i].InverseTransformVector(ClickPos.WorldRay.Direction));
+
+		if (MeshSpatials[i]->FindNearestHitTriangle(LocalRay, RayT, HitTid))
 		{
-			CurrentSelection.Mesh = Spatial->GetMesh();
+			CurrentSelection.Mesh = MeshSpatials[i]->GetMesh();
+			CurrentSelectionIndex = i;
 			break;
 		}
 	}
@@ -156,6 +270,8 @@ void UMeshSelectionMechanic::OnClicked(const FInputDeviceRay& ClickPos)
 
 	if (!(OriginalSelection == CurrentSelection))
 	{
+		UpdateCentroid();
+		RebuildDrawnElements(FTransform(CurrentSelectionCentroid));
 		OnSelectionChanged.Broadcast();
 	}
 }
diff --git a/Engine/Plugins/Experimental/MeshModelingToolset/Source/ModelingComponents/Public/Selection/DynamicMeshSelection.h b/Engine/Plugins/Experimental/MeshModelingToolset/Source/ModelingComponents/Public/Selection/DynamicMeshSelection.h
--- a/Engine/Plugins/Experimental/MeshModelingToolset/Source/ModelingComponents/Public/Selection/DynamicMeshSelection.h
+++ b/Engine/Plugins/Experimental/MeshModelingToolset/Source/ModelingComponents/Public/Selection/DynamicMeshSelection.h
@@ -27,7 +27,7 @@ public:
 	TSet<int32> SelectedIDs;
 	EType Type = EType::Vertex;
 
-	// Not be relevant if the selection type is not Group
+	// Not relevant if the selection type is not Group
 	int32 GroupLayer = 0;
 
 	bool IsEmpty() const
diff --git a/Engine/Plugins/Experimental/MeshModelingToolset/Source/ModelingComponents/Public/Selection/MeshSelectionMechanic.h b/Engine/Plugins/Experimental/MeshModelingToolset/Source/ModelingComponents/Public/Selection/MeshSelectionMechanic.h
--- a/Engine/Plugins/Experimental/MeshModelingToolset/Source/ModelingComponents/Public/Selection/MeshSelectionMechanic.h
+++ b/Engine/Plugins/Experimental/MeshModelingToolset/Source/ModelingComponents/Public/Selection/MeshSelectionMechanic.h
@@ -14,6 +14,9 @@
 
 #include "MeshSelectionMechanic.generated.h"
 
+class ULineSetComponent;
+class APreviewGeometryActor;
+
 enum class EMeshSelectionMechanicMode
 {
 	Component,
@@ -50,8 +53,20 @@ public:
 	virtual void Setup(UInteractiveTool* ParentTool) override;
 	virtual void Shutdown() override;
 
+	void SetWorld(UWorld* World);
+
 	// Use this to initialize the meshes we want to hit test.
-	virtual void AddSpatial(TSharedPtr<FDynamicMeshAABBTree3> SpatialIn);
+	virtual void AddSpatial(TSharedPtr<FDynamicMeshAABBTree3> SpatialIn, const FTransform& TransformIn);
+
+	FVector3d GetCurrentSelectionCentroid();
+
+	// Rebuilds the drawn selection highlights, and intializes them in such a way that their transform
+	// is equal to StartTransform (useful so that their transform can later be changed)
+	void RebuildDrawnElements(const FTransform& StartTransform);
+
+	// Changes the transform of the selection highlights. Useful for quickly updating the hightlight
+	// without rebuilding it, when the change is a transformation.
+	void SetDrawnElementsTransform(const FTransform& Transform);
 
 	virtual const FDynamicMeshSelection& GetCurrentSelection() const;
 	virtual void SetSelection(const FDynamicMeshSelection& Selection, bool bBroadcast = false, bool bEmitChange = false);
@@ -68,8 +83,25 @@ public:
 	EMeshSelectionMechanicMode SelectionMode;
 
 protected:
+
+	UPROPERTY()
+	APreviewGeometryActor* PreviewGeometryActor = nullptr;
+
+	UPROPERTY()
+	ULineSetComponent* LineSet = nullptr;
+
+
 	TArray<TSharedPtr<FDynamicMeshAABBTree3>> MeshSpatials;
+	TArray<FTransform> MeshTransforms;
 	FDynamicMeshSelection CurrentSelection;
+	int32 CurrentSelectionIndex;
 	FViewCameraState CameraState;
+
+	FColor LineColor = FColor::Yellow;
+	float LineThickness = 3;
+	float DepthBias = 0.3;
+
+	FVector3d CurrentSelectionCentroid;
+	void UpdateCentroid();
 };
 
diff --git a/Engine/Plugins/Experimental/UVEditor/Source/UVEditor/Private/ToolTargets/StaticMeshUVMeshToolTarget.cpp b/Engine/Plugins/Experimental/UVEditor/Source/UVEditor/Private/ToolTargets/StaticMeshUVMeshToolTarget.cpp
deleted file mode 100644
diff --git a/Engine/Plugins/Experimental/UVEditor/Source/UVEditor/Private/UVEditorMode.cpp b/Engine/Plugins/Experimental/UVEditor/Source/UVEditor/Private/UVEditorMode.cpp
--- a/Engine/Plugins/Experimental/UVEditor/Source/UVEditor/Private/UVEditorMode.cpp
+++ b/Engine/Plugins/Experimental/UVEditor/Source/UVEditor/Private/UVEditorMode.cpp
@@ -3,23 +3,31 @@
 #include "UVEditorMode.h"
 
 #include "Algo/AnyOf.h"
+#include "ContextObjectStore.h"
 #include "Drawing/MeshElementsVisualizer.h"
 #include "EdModeInteractiveToolsContext.h" //ToolsContext
 #include "InteractiveTool.h"
-#include "TargetInterfaces/UVUnwrapDynamicMesh.h"
-#include "ToolTargetManager.h"
-#include "ToolTargets/StaticMeshUVMeshToolTarget.h"
-#include "ToolTargets/ToolTarget.h"
+#include "MeshOpPreviewHelpers.h" // UMeshOpPreviewWithBackgroundCompute
 #include "PreviewMesh.h"
+#include "TargetInterfaces/AssetBackedTarget.h"
+#include "TargetInterfaces/MaterialProvider.h"
+#include "TargetInterfaces/DynamicMeshProvider.h"
+#include "TargetInterfaces/DynamicMeshCommitter.h"
+#include "ToolSetupUtil.h"
+#include "ToolTargets/UVEditorToolMeshInput.h"
+#include "ToolTargetManager.h"
+#include "ToolTargets/UVEditorToolMeshInput.h"
 #include "UVEditorCommands.h"
 #include "UVSelectTool.h"
 #include "UVEditorModeToolkit.h"
 #include "UVEditorSubsystem.h"
-#include "ToolSetupUtil.h"
-#include "UVToolStateObjects.h"
+#include "UVEditorToolUtil.h"
+#include "UVToolContextObjects.h"
 
 #define LOCTEXT_NAMESPACE "UUVEditorMode"
 
+using namespace UE::Geometry;
+
 const FEditorModeID UUVEditorMode::EM_UVEditorModeId = TEXT("EM_UVEditorMode");
 
 namespace UVEditorModeLocals
@@ -33,6 +41,18 @@ namespace UVEditorModeLocals
 	FColor IslandBorderColor = FColor(103, 52, 235);
 }
 
+const FToolTargetTypeRequirements& UUVEditorMode::GetToolTargetRequirements()
+{
+	static const FToolTargetTypeRequirements ToolTargetRequirements =
+		FToolTargetTypeRequirements({
+			UMaterialProvider::StaticClass(),
+			UDynamicMeshCommitter::StaticClass(),
+			UDynamicMeshProvider::StaticClass(),
+			UAssetBackedTarget::StaticClass()
+			});
+	return ToolTargetRequirements;
+}
+
 UUVEditorMode::UUVEditorMode()
 {
 	Info = FEditorModeInfo(
@@ -46,14 +66,9 @@ void UUVEditorMode::Enter()
 {
 	Super::Enter();
 
-	StateObjectStore = NewObject<UUVToolStateObjectStore>(this);
-
 	RegisterTools();
 
 	ActivateDefaultTool();
-
-	// Add new target factories here and in UUVEditorSubsystem::Initialize() as they are developed.
-	ToolsContext->TargetManager->AddTargetFactory(NewObject<UStaticMeshUVMeshToolTargetFactory>(ToolsContext->TargetManager));
 }
 
 void UUVEditorMode::RegisterTools()
@@ -62,13 +77,11 @@ void UUVEditorMode::RegisterTools()
 
 	// TODO: Other tool registrations go here
 	auto UVSelectToolBuilder = NewObject<UUVSelectToolBuilder>();
-	UVSelectToolBuilder->DisplayedMeshes = &DisplayedMeshes;
-	UVSelectToolBuilder->StateObjectStore = GetStateObjectStore();
+	UVSelectToolBuilder->Targets = &ToolInputObjects;
 	RegisterTool(CommandInfos.BeginSelectTool, TEXT("UVSelectTool"), UVSelectToolBuilder);
 
 	auto UVTransformToolBuilder = NewObject<UUVSelectToolBuilder>();
-	UVTransformToolBuilder->DisplayedMeshes = &DisplayedMeshes;
-	UVTransformToolBuilder->StateObjectStore = GetStateObjectStore();
+	UVTransformToolBuilder->Targets = &ToolInputObjects;
 	UVTransformToolBuilder->bGizmoEnabled = true;
 	RegisterTool(CommandInfos.BeginTransformTool, TEXT("UVTransformTool"), UVTransformToolBuilder);
 }
@@ -92,55 +105,80 @@ void UUVEditorMode::BindCommands()
 
 void UUVEditorMode::Exit()
 {
-	StateObjectStore->Clear();
-	StateObjectStore = nullptr;
-
-	for (TObjectPtr<UPreviewMesh>& Preview : DisplayedMeshes)
+	for (TObjectPtr<UUVEditorToolMeshInput> ToolInput : ToolInputObjects)
 	{
-		Preview->Disconnect();
-	}
-	for (TObjectPtr<UMeshElementsVisualizer>& Wireframe : DisplayedWireframes)
-	{
-		Wireframe->Disconnect();
+		ToolInput->Shutdown();
 	}
+	ToolInputObjects.Reset();
 
 	Super::Exit();
 }
 
-void UUVEditorMode::InitializeTargets(TArray<TObjectPtr<UObject>>& AssetsIn)
+void UUVEditorMode::InitializeTargets(TArray<TObjectPtr<UObject>>& AssetsIn, TArray<FTransform>* TransformsIn)
 {
 	using namespace UVEditorModeLocals;
 
 	OriginalObjectsToEdit = AssetsIn;
 
-	for (const TObjectPtr<UObject>& Object : AssetsIn)
+	// Build the tool targets that provide us with 3d dynamic meshes
+	UUVEditorSubsystem* UVSubsystem = GEditor->GetEditorSubsystem<UUVEditorSubsystem>();
+	UVSubsystem->BuildTargets(AssetsIn, GetToolTargetRequirements(), ToolTargets);
+
+	// Get an api for manipulating things in the live preview, which has its own world and input router
+	UContextObjectStore* ContextStore = GetInteractiveToolsContext()->ToolManager->GetContextObjectStore();
+	UUVToolLivePreviewAPI* LivePreviewAPI = ContextStore->FindContext<UUVToolLivePreviewAPI>();
+	check(LivePreviewAPI);
+	UWorld* LivePreviewWorld = LivePreviewAPI->GetLivePreviewWorld();
+
+	double ScaleFactor = GetUVMeshScalingFactor();
+
+	// These functions will determine the mapping between UV values and the resulting mesh vertex positions. 
+	// If we're looking down on the unwrapped mesh, with the Z axis towards us, we want U's to be right, and
+	// V's to be up. In Unreal's left-handed coordinate system, this means that we map U's to world Y
+	// and V's to world X.
+	// Also, Unreal changes the V coordinates of imported meshes to 1-V internally, and we undo this
+	// while displaying the UV's because the users likely expect to see the original UV's (it would
+	// be particularly confusing for users working with UDIM assets, where internally stored V's 
+	// frequently end up negative).
+	// The ScaleFactor just scales the mesh up. Scaling the mesh up makes it easier to zoom in
+	// further into the display before getting issues with the camera near plane distance.
+	auto UVToVertPosition = [this, ScaleFactor](const FVector2f& UV)
+	{
+		return FVector3d((1 - UV.Y) * ScaleFactor, UV.X * ScaleFactor, 0);
+	};
+	auto VertPositionToUV = [this, ScaleFactor](const FVector3d& VertPosition)
 	{
-		UToolTarget* ToolTarget = ToolsContext->TargetManager->BuildTarget(Object, UUVEditorSubsystem::UVUnwrapMeshTargetRequirements);
-		IUVUnwrapDynamicMesh* UVMesh = Cast<IUVUnwrapDynamicMesh>(ToolTarget);
-		check(UVMesh);
-		UVUnwrapMeshTargets.Add(ToolTarget);
-
-		// Create in-viewport representations of the unwrapped mesh
-		UPreviewMesh* UVPreview = NewObject<UPreviewMesh>();
-		UVPreview->CreateInWorld(GetWorld(), FTransform::Identity);
-		UVPreview->UpdatePreview(UVMesh->GetMesh(UVLayerIndex).Get());
+		return FVector2D(VertPosition.Y / ScaleFactor, 1 - (VertPosition.X / ScaleFactor));
+	};
+
+	// For each of our tool targets, we're going to create a UV tool input object that bundles together
+	// the 3d preview and the unwrapped UV layer, with both a reference and a "preview with background op"
+	// versions of both
+	for (UToolTarget* Target : ToolTargets)
+	{
+		UUVEditorToolMeshInput* ToolInputObject = NewObject<UUVEditorToolMeshInput>();
 		
-		// We use a two-sided material because the triangles of our unwrapped meshes may have different orientations
-		UVPreview->SetMaterial(0, ToolSetupUtil::GetCustomTwoSidedDepthOffsetMaterial(
-			GetToolManager(),
-			(FLinearColor)TriangleColor,
-			0, //depth offset
-			0.3)); //opacity
+		if (!ToolInputObject->InitializeMeshes(Target, UVLayerIndex,
+			GetWorld(), LivePreviewWorld, ToolSetupUtil::GetDefaultWorkingMaterial(GetToolManager()),
+			UVToVertPosition, VertPositionToUV))
+		{
+			continue;
+		}
 
-		UVPreview->SetVisible(true);
-		DisplayedMeshes.Add(UVPreview);
+		ToolInputObject->UnwrapPreview->PreviewMesh->SetMaterial(
+			0, ToolSetupUtil::GetCustomTwoSidedDepthOffsetMaterial(
+				GetToolManager(),
+				(FLinearColor)TriangleColor,
+				0)); //depth offset
+
+		// Initialize our timestamp so we can later detect changes
+		MeshChangeStamps.Add(ToolInputObject->UnwrapCanonical->GetShapeTimestamp());
 
 		// Set up the wireframe display of the unwrapped mesh.
 		UMeshElementsVisualizer* WireframeDisplay = NewObject<UMeshElementsVisualizer>(this);
 		WireframeDisplay->CreateInWorld(GetWorld(), FTransform::Identity);
-		check(WireframeDisplay->Settings);
 
-		WireframeDisplay->Settings->DepthBias = 1;
+		WireframeDisplay->Settings->DepthBias = 0.5;
 		WireframeDisplay->Settings->bAdjustDepthBiasUsingMeshSize = false;
 		WireframeDisplay->Settings->bShowWireframe = true;
 		WireframeDisplay->Settings->bShowBorders = true;
@@ -148,33 +186,45 @@ void UUVEditorMode::InitializeTargets(TArray<TObjectPtr<UObject>>& AssetsIn)
 		WireframeDisplay->Settings->BoundaryEdgeColor = IslandBorderColor;
 		WireframeDisplay->Settings->bShowUVSeams = false;
 		WireframeDisplay->Settings->bShowNormalSeams = false;
-
 		// These are not exposed at the visualizer level yet
 		// TODO: Should they be?
 		WireframeDisplay->WireframeComponent->BoundaryEdgeThickness = 2;
 
-		// The settings object is not part of a tool, so it won't get ticked like its
-		// supposed to (to enable property watching), unless we add this here.
+		// The wireframe will track the unwrap preview mesh
+		WireframeDisplay->SetMeshAccessFunction([ToolInputObject](void) { 
+			return ToolInputObject->UnwrapPreview->PreviewMesh->GetMesh(); 
+			});
+
+		// The settings object and wireframe are not part of a tool, so they won't get ticked like they
+		// are supposed to (to enable property watching), unless we add this here.
 		PropertyObjectsToTick.Add(WireframeDisplay->Settings);
+		WireframesToTick.Add(WireframeDisplay);
 
-		// The wireframe will track the preview mesh
-		WireframeDisplay->SetMeshAccessFunction([UVPreview](void) { return UVPreview->GetMesh(); });
-		UVPreview->GetOnMeshChanged().AddWeakLambda(this, [WireframeDisplay]() 
-		{
-			WireframeDisplay->NotifyMeshChanged();
-		});
-		DisplayedWireframes.Add(WireframeDisplay);
+		// The tool input object will hold on to the wireframe for the purposes of updating it and cleaning it up
+		ToolInputObject->WireframeDisplay = WireframeDisplay;
 
-		// Initialize our timestamp so we can later detect changes
-		MeshChangeStamps.Add(UVPreview->GetMesh()->GetShapeTimestamp());
+		ToolInputObjects.Add(ToolInputObject);
+	}
+	
+	if (TransformsIn && TransformsIn->Num() == AssetsIn.Num())
+	{
+		for (int i = 0; i < ToolInputObjects.Num(); ++i)
+		{
+			ToolInputObjects[i]->AppliedPreview->PreviewMesh->SetTransform((*TransformsIn)[i]);
+		}
 	}
 }
 
+void UUVEditorMode::EmitToolIndependentObjectChange(UObject* TargetObject, TUniquePtr<FToolCommandChange> Change, const FText& Description)
+{
+	GetInteractiveToolsContext()->GetTransactionAPI()->AppendChange(TargetObject, MoveTemp(Change), Description);
+}
+
 bool UUVEditorMode::HaveUnappliedChanges()
 {
-	for (int32 i = 0; i < DisplayedMeshes.Num(); ++i)
+	for (int32 i = 0; i < ToolInputObjects.Num(); ++i)
 	{
-		if (DisplayedMeshes[i]->GetMesh()->GetShapeTimestamp() != MeshChangeStamps[i])
+		if (ToolInputObjects[i]->UnwrapCanonical->GetShapeTimestamp() != MeshChangeStamps[i])
 		{
 			return true;
 		}
@@ -184,9 +234,9 @@ bool UUVEditorMode::HaveUnappliedChanges()
 
 void UUVEditorMode::GetAssetsWithUnappliedChanges(TArray<TObjectPtr<UObject>> UnappliedAssetsOut)
 {
-	for (int32 i = 0; i < DisplayedMeshes.Num(); ++i)
+	for (int32 i = 0; i < ToolInputObjects.Num(); ++i)
 	{
-		if (DisplayedMeshes[i]->GetMesh()->GetShapeTimestamp() != MeshChangeStamps[i])
+		if (ToolInputObjects[i]->UnwrapCanonical->GetShapeTimestamp() != MeshChangeStamps[i])
 		{
 			UnappliedAssetsOut.Add(OriginalObjectsToEdit[i]);
 		}
@@ -197,14 +247,21 @@ void UUVEditorMode::ApplyChanges()
 {
 	using namespace UVEditorModeLocals;
 
-	for (int32 i = 0; i < DisplayedMeshes.Num(); ++i)
+	GetToolManager()->BeginUndoTransaction(LOCTEXT("UVEditorApplyChangesTransaction", "UV Editor Apply Changes"));
+
+	IDynamicMeshCommitter::FDynamicMeshCommitInfo CommitInfo(false);
+	CommitInfo.bUVsChanged = true;
+
+	for (int32 i = 0; i < ToolInputObjects.Num(); ++i)
 	{
-		if (DisplayedMeshes[i]->GetMesh()->GetShapeTimestamp() != MeshChangeStamps[i])
+		if (ToolInputObjects[i]->UnwrapCanonical->GetShapeTimestamp() != MeshChangeStamps[i])
 		{
-			UVUnwrapMeshTargets[i]->SaveBackToUVs(DisplayedMeshes[i]->GetMesh(), UVLayerIndex);
-			MeshChangeStamps[i] = DisplayedMeshes[i]->GetMesh()->GetShapeTimestamp();
+			Cast<IDynamicMeshCommitter>(ToolTargets[i])->CommitDynamicMesh(*ToolInputObjects[i]->AppliedCanonical, CommitInfo);
+			MeshChangeStamps[i] = ToolInputObjects[i]->UnwrapCanonical->GetShapeTimestamp();
 		}
 	}
+
+	GetToolManager()->EndUndoTransaction();
 }
 
 void UUVEditorMode::ModeTick(float DeltaTime)
@@ -226,9 +283,12 @@ void UUVEditorMode::ModeTick(float DeltaTime)
 		}
 	}
 
-	for (UMeshElementsVisualizer* WireframeDisplay : DisplayedWireframes)
+	for (TWeakObjectPtr<UMeshElementsVisualizer> WireframeDisplay : WireframesToTick)
 	{
-		WireframeDisplay->OnTick(DeltaTime);
+		if (WireframeDisplay.IsValid())
+		{
+			WireframeDisplay->OnTick(DeltaTime);
+		}
 	}
 }
 
diff --git a/Engine/Plugins/Experimental/UVEditor/Source/UVEditor/Private/UVEditorSubsystem.cpp b/Engine/Plugins/Experimental/UVEditor/Source/UVEditor/Private/UVEditorSubsystem.cpp
--- a/Engine/Plugins/Experimental/UVEditor/Source/UVEditor/Private/UVEditorSubsystem.cpp
+++ b/Engine/Plugins/Experimental/UVEditor/Source/UVEditor/Private/UVEditorSubsystem.cpp
@@ -3,12 +3,12 @@
 #include "UVEditorSubsystem.h"
 
 #include "ToolTargetManager.h"
-#include "ToolTargets/StaticMeshUVMeshToolTarget.h"
-#include "TargetInterfaces/UVUnwrapDynamicMesh.h"
+#include "ToolTargets/StaticMeshToolTarget.h"
+#include "ToolTargets/SkeletalMeshToolTarget.h"
 #include "UVEditor.h"
+#include "UVEditorMode.h"
 
-const FToolTargetTypeRequirements UUVEditorSubsystem::UVUnwrapMeshTargetRequirements = 
-	FToolTargetTypeRequirements(UUVUnwrapDynamicMesh::StaticClass());
+using namespace UE::Geometry;
 
 void UUVEditorSubsystem::Initialize(FSubsystemCollectionBase& Collection)
 {
@@ -17,8 +17,8 @@ void UUVEditorSubsystem::Initialize(FSubsystemCollectionBase& Collection)
 	ToolTargetManager = NewObject<UToolTargetManager>(this);
 	ToolTargetManager->Initialize();
 
-	// Add new target factories here and in UUVEditorMode::Enter() as they are developed.
-	ToolTargetManager->AddTargetFactory(NewObject<UStaticMeshUVMeshToolTargetFactory>(ToolTargetManager));
+	ToolTargetManager->AddTargetFactory(NewObject<UStaticMeshToolTargetFactory>(ToolTargetManager));
+	ToolTargetManager->AddTargetFactory(NewObject<USkeletalMeshToolTargetFactory>(ToolTargetManager));
 }
 
 void UUVEditorSubsystem::Deinitialize()
@@ -36,7 +36,7 @@ bool UUVEditorSubsystem::AreObjectsValidTargets(const TArray<UObject*>& InObject
 
 	for (UObject* Object : InObjects)
 	{
-		if (!ToolTargetManager->CanBuildTarget(Object, UVUnwrapMeshTargetRequirements))
+		if (!ToolTargetManager->CanBuildTarget(Object, UUVEditorMode::GetToolTargetRequirements()))
 		{
 			return false;
 		}
@@ -45,6 +45,21 @@ bool UUVEditorSubsystem::AreObjectsValidTargets(const TArray<UObject*>& InObject
 	return true;
 }
 
+void UUVEditorSubsystem::BuildTargets(const TArray<TObjectPtr<UObject>>& ObjectsIn, 
+	const FToolTargetTypeRequirements& TargetRequirements, TArray<TObjectPtr<UToolTarget>>& TargetsOut)
+{
+	TargetsOut.Reset();
+
+	for (UObject* Object : ObjectsIn)
+	{
+		UToolTarget* Target = ToolTargetManager->BuildTarget(Object, TargetRequirements);
+		if (Target)
+		{
+			TargetsOut.Add(Target);
+		}
+	}
+}
+
 void UUVEditorSubsystem::StartUVEditor(TArray<TObjectPtr<UObject>> ObjectsToEdit)
 {
 	// We don't allow opening a new instance if any of the objects are already opened
diff --git a/Engine/Plugins/Experimental/UVEditor/Source/UVEditor/Private/UVEditorToolkit.cpp b/Engine/Plugins/Experimental/UVEditor/Source/UVEditor/Private/UVEditorToolkit.cpp
--- a/Engine/Plugins/Experimental/UVEditor/Source/UVEditor/Private/UVEditorToolkit.cpp
+++ b/Engine/Plugins/Experimental/UVEditor/Source/UVEditor/Private/UVEditorToolkit.cpp
@@ -2,16 +2,21 @@
 
 #include "UVEditorToolkit.h"
 
+#include "AdvancedPreviewScene.h"
 #include "AssetEditorModeManager.h"
+#include "ContextObjectStore.h"
+#include "EditorViewportTabContent.h"
 #include "Misc/MessageDialog.h"
 #include "PreviewScene.h"
 #include "SAssetEditorViewport.h"
+#include "ThumbnailRendering/SceneThumbnailInfo.h"
 #include "ToolMenus.h"
 #include "UVEditor.h"
 #include "UVEditorMode.h"
 #include "UVEditorCommands.h"
 #include "UVEditorSubsystem.h"
 #include "UVEditor2DViewportClient.h"
+#include "UVToolContextObjects.h"
 #include "Widgets/Docking/SDockTab.h"
 
 #include "EdModeInteractiveToolsContext.h"
@@ -19,18 +24,17 @@
 #define LOCTEXT_NAMESPACE "UVEditorToolkit"
 
 const FName FUVEditorToolkit::InteractiveToolsPanelTabID(TEXT("UVEditor_InteractiveToolsTab"));
+const FName FUVEditorToolkit::LivePreviewTabID(TEXT("UVEditor_LivePreviewTab"));
 
 FUVEditorToolkit::FUVEditorToolkit(UAssetEditor* InOwningAssetEditor)
 	: FBaseAssetToolkit(InOwningAssetEditor)
 {
 	check(Cast<UUVEditor>(InOwningAssetEditor));
 
-	// TODO: Will eventually need another viewport/previewscene for a 3d preview of the mesh
-
 	// We will replace the StandaloneDefaultLayout that our parent class gave us with 
-	// one where the properties detail panel is a vertical column on the left, and the
-	// viewport is on the right.
-	StandaloneDefaultLayout = FTabManager::NewLayout(FName("UVEditorLayout"))
+	// one where the properties detail panel is a vertical column on the left, and there
+	// are two viewports on the right.
+	StandaloneDefaultLayout = FTabManager::NewLayout(FName("UVEditorLayout1"))
 		->AddArea
 		(
 			FTabManager::NewPrimaryArea()->SetOrientation(Orient_Vertical)
@@ -42,27 +46,54 @@ FUVEditorToolkit::FUVEditorToolkit(UAssetEditor* InOwningAssetEditor)
 					FTabManager::NewStack()
 					->SetSizeCoefficient(0.2f)
 					->AddTab(InteractiveToolsPanelTabID, ETabState::OpenedTab)
-					->SetHideTabWell(true)
+					->SetHideTabWell(false)
 				)
 				->Split
 				(
 					FTabManager::NewStack()
-					->SetSizeCoefficient(0.7f)
+					->SetSizeCoefficient(0.4f)
 					->AddTab(ViewportTabID, ETabState::OpenedTab)
 					->SetHideTabWell(true)
 				)
+				->Split
+				(
+					FTabManager::NewStack()
+					->SetSizeCoefficient(0.4f)
+					->AddTab(LivePreviewTabID, ETabState::OpenedTab)
+					->SetHideTabWell(false)
+				)
 			)
 		);
 
-	// We could create the preview scene in CreateEditorViewportClient() the way that FBaseAssetToolkit
-	// does, but it seems more intuitive to create it right off the bat and pass it in later. The preview
-	// scene is important because it holds the world in which our mode will create any objects to display.
+	// We could create the preview scenes in CreateEditorViewportClient() the way that FBaseAssetToolkit
+	// does, but it seems more intuitive to create them right off the bat and pass it in later. 
 	FPreviewScene::ConstructionValues PreviewSceneArgs;
-	PreviewScene = MakeUnique<FPreviewScene>(PreviewSceneArgs);
+	UnwrapScene = MakeUnique<FPreviewScene>(PreviewSceneArgs);
+	LivePreviewScene = MakeUnique<FAdvancedPreviewScene>(PreviewSceneArgs);
+	LivePreviewScene->SetFloorVisibility(false, true);
+
+	LivePreviewEditorModeManager = MakeShared<FAssetEditorModeManager>();
+	LivePreviewEditorModeManager->SetPreviewScene(LivePreviewScene.Get());
+	LivePreviewInputRouter = LivePreviewEditorModeManager->GetInteractiveToolsContext()->InputRouter;
+
+	LivePreviewTabContent = MakeShareable(new FEditorViewportTabContent());
+	LivePreviewViewportClient = MakeShared<FEditorViewportClient>(
+		LivePreviewEditorModeManager.Get(), LivePreviewScene.Get());
+
+	LivePreviewViewportDelegate = [this](FAssetEditorViewportConstructionArgs InArgs)
+	{
+		return SNew(SAssetEditorViewport, InArgs)
+			.EditorViewportClient(LivePreviewViewportClient);
+	};
 }
 
 FUVEditorToolkit::~FUVEditorToolkit()
 {
+	// We need to force the uv editor mode deletion now because otherwise the preview and unwrap worlds
+	// will end up getting destroyed before the mode's Exit() function gets to run, and we'll get some
+	// warnings when we destroy any mode actors.
+	EditorModeManager->DestroyMode(UUVEditorMode::EM_UVEditorModeId);
+
 	// The UV subsystem is responsible for opening/focusing UV editor instances, so we should
 	// notify it that this one is closing.
 	UUVEditorSubsystem* UVSubsystem = GEditor->GetEditorSubsystem<UUVEditorSubsystem>();
@@ -119,16 +150,28 @@ FText FUVEditorToolkit::GetToolkitToolTipText() const
 
 void FUVEditorToolkit::RegisterTabSpawners(const TSharedRef<FTabManager>& InTabManager)
 {
-	// Here we set up the tabs we referenced in StandaloneDefaultLayout (in the constructor).
-	// The base class already creates the viewport tab (and also a details tab that we don't want).
+	// We bypass FBaseAssetToolkit::RegisterTabSpawners because it doesn't seem to provide us with
+	// anything except tabs that we don't want.
+	FAssetEditorToolkit::RegisterTabSpawners(InTabManager);
 
-	FBaseAssetToolkit::RegisterTabSpawners(InTabManager);
-	InTabManager->UnregisterTabSpawner(DetailsTabID);
+	// Here we set up the tabs we referenced in StandaloneDefaultLayout (in the constructor).
 
-	InTabManager->RegisterTabSpawner(InteractiveToolsPanelTabID, FOnSpawnTab::CreateSP(this, &FUVEditorToolkit::SpawnTab_InteractiveToolsPanel))
+	InTabManager->RegisterTabSpawner(InteractiveToolsPanelTabID, FOnSpawnTab::CreateSP(this, 
+		&FUVEditorToolkit::SpawnTab_InteractiveToolsPanel))
 		.SetDisplayName(LOCTEXT("InteractiveToolsPanel", "Tools Panel"))
 		.SetGroup(AssetEditorTabsCategory.ToSharedRef())
 		.SetIcon(FSlateIcon(FEditorStyle::GetStyleSetName(), "LevelEditor.Tabs.Details"));
+
+	InTabManager->RegisterTabSpawner(ViewportTabID, FOnSpawnTab::CreateSP(this, &FUVEditorToolkit::SpawnTab_Viewport))
+		.SetDisplayName(LOCTEXT("ViewportTab", "Viewport"))
+		.SetGroup(AssetEditorTabsCategory.ToSharedRef())
+		.SetIcon(FSlateIcon(FEditorStyle::GetStyleSetName(), "LevelEditor.Tabs.Viewports"));
+
+	InTabManager->RegisterTabSpawner(LivePreviewTabID, FOnSpawnTab::CreateSP(this, 
+		&FUVEditorToolkit::SpawnTab_LivePreview))
+		.SetDisplayName(LOCTEXT("ViewportTab", "Viewport"))
+		.SetGroup(AssetEditorTabsCategory.ToSharedRef())
+		.SetIcon(FSlateIcon(FEditorStyle::GetStyleSetName(), "LevelEditor.Tabs.Viewports"));
 }
 
 bool FUVEditorToolkit::OnRequestClose()
@@ -182,7 +225,17 @@ void FUVEditorToolkit::SaveAsset_Execute()
 	FAssetEditorToolkit::SaveAsset_Execute();
 }
 
-// This is bound in RegisterTabSpawners() to create the panel on the left. It's filled in by the mode.
+TSharedRef<SDockTab> FUVEditorToolkit::SpawnTab_LivePreview(const FSpawnTabArgs& Args)
+{
+	TSharedRef< SDockTab > DockableTab =
+		SNew(SDockTab);
+
+	const FString LayoutId = FString("UVEditorLivePreviewViewport");
+	LivePreviewTabContent->Initialize(LivePreviewViewportDelegate, DockableTab, LayoutId);
+	return DockableTab;
+}
+
+// This is bound in RegisterTabSpawners() to create the panel on the left. The panel is filled in by the mode.
 TSharedRef<SDockTab> FUVEditorToolkit::SpawnTab_InteractiveToolsPanel(const FSpawnTabArgs& Args)
 {
 	ToolsPanel = SNew(SDockTab)
@@ -191,18 +244,18 @@ TSharedRef<SDockTab> FUVEditorToolkit::SpawnTab_InteractiveToolsPanel(const FSpa
 	return ToolsPanel.ToSharedRef();
 }
 
-// This is used to create the viewport itself
-AssetEditorViewportFactoryFunction FUVEditorToolkit::GetViewportDelegate()
+void FUVEditorToolkit::CreateWidgets()
 {
-	AssetEditorViewportFactoryFunction TempViewportDelegate = [this](FAssetEditorViewportConstructionArgs InArgs)
-	{
-		return SNew(SAssetEditorViewport, InArgs)
-			.EditorViewportClient(ViewportClient);
-	};
+	// This gets called during UAssetEditor::Init() after creation of the toolkit but before
+	// calling InitAssetEditor on it. If we have custom mode-level toolbars we want to add,
+	// they could potentially go here, but we still need to call the base CreateWidgets as well
+	// because that calls things that make viewport client, etc.
 
-	return TempViewportDelegate;
+	FBaseAssetToolkit::CreateWidgets();
 }
 
+// Called from FBaseAssetToolkit::CreateWidgets to populate ViewportClient, but otherwise only used 
+// in our own viewport delegate.
 TSharedPtr<FEditorViewportClient> FUVEditorToolkit::CreateEditorViewportClient() const
 {
 	// Note that we can't reliably adjust the viewport client here because we will be passing it
@@ -211,19 +264,23 @@ TSharedPtr<FEditorViewportClient> FUVEditorToolkit::CreateEditorViewportClient()
 	// namely ViewportType.
 	// Instead, we do viewport client adjustment in PostInitAssetEditor().
 	check(EditorModeManager.IsValid());
-	return MakeShared<FUVEditor2DViewportClient>(EditorModeManager.Get(), PreviewScene.Get());
+	return MakeShared<FUVEditor2DViewportClient>(EditorModeManager.Get(), UnwrapScene.Get());
 }
 
-void FUVEditorToolkit::CreateWidgets()
+// Called from FBaseAssetToolkit::CreateWidgets. The delegate call path goes through FAssetEditorToolkit::InitAssetEditor
+// and FBaseAssetToolkit::SpawnTab_Viewport.
+AssetEditorViewportFactoryFunction FUVEditorToolkit::GetViewportDelegate()
 {
-	// This gets called during UAssetEditor::Init() after creation of the toolkit but before
-	// calling InitAssetEditor on it. If we have custom mode-level toolbars we want to add,
-	// they could potentially go here, but we still need to call the base CreateWidgets as well.
+	AssetEditorViewportFactoryFunction TempViewportDelegate = [this](FAssetEditorViewportConstructionArgs InArgs)
+	{
+		return SNew(SAssetEditorViewport, InArgs)
+			.EditorViewportClient(ViewportClient);
+	};
 
-	FBaseAssetToolkit::CreateWidgets();
+	return TempViewportDelegate;
 }
 
-
+// Called from FBaseAssetToolkit::CreateWidgets.
 void FUVEditorToolkit::CreateEditorModeManager()
 {
 	EditorModeManager = MakeShared<FAssetEditorModeManager>();
@@ -232,7 +289,7 @@ void FUVEditorToolkit::CreateEditorModeManager()
 	// and setting the preview scene here makes our GetWorld() function return the preview scene
 	// world instead of the normal level editor one. Important because that is where we create
 	// any preview meshes, gizmo actors, etc.
-	StaticCastSharedPtr<FAssetEditorModeManager>(EditorModeManager)->SetPreviewScene(PreviewScene.Get());
+	StaticCastSharedPtr<FAssetEditorModeManager>(EditorModeManager)->SetPreviewScene(UnwrapScene.Get());
 }
 
 void FUVEditorToolkit::PostInitAssetEditor()
@@ -247,10 +304,21 @@ void FUVEditorToolkit::PostInitAssetEditor()
 		EditorModeManager->GetActiveScriptableMode(UUVEditorMode::EM_UVEditorModeId));
 	check(UVMode);
 
+	// The mode will need to be able to get to the live preview world and input.
+	UContextObjectStore* ContextStore = EditorModeManager->GetInteractiveToolsContext()->ToolManager->GetContextObjectStore();
+	UUVToolLivePreviewAPI* LivePreviewAPI = NewObject<UUVToolLivePreviewAPI>();
+	LivePreviewAPI->Initialize(LivePreviewScene->GetWorld(), LivePreviewInputRouter);
+	ContextStore->AddContextObject(LivePreviewAPI);
+
 	// Initialize mode state.
 	TArray<TObjectPtr<UObject>> ObjectsToEdit;
 	OwningAssetEditor->GetObjectsToEdit(ObjectsToEdit);
-	UVMode->InitializeTargets(ObjectsToEdit);
+
+	// TODO: get these when possible, set them otherwise.
+	TArray<FTransform> ObjectTransforms;
+	ObjectTransforms.SetNum(ObjectsToEdit.Num());
+
+	UVMode->InitializeTargets(ObjectsToEdit, &ObjectTransforms);
 
 	// Plug in the mode tool panel
 	TSharedPtr<FModeToolkit> UVModeToolkit = UVMode->GetToolkit().Pin();
@@ -278,7 +346,7 @@ void FUVEditorToolkit::PostInitAssetEditor()
 	Section.AddEntry(FToolMenuEntry::InitToolBarButton(FUVEditorCommands::Get().ApplyChanges));
 
 
-	// Adjust the viewport:
+	// Adjust our main (2D) viewport:
 
 	// Ortho has too many problems with rendering things, unfortunately, so we should use perspective.
 	ViewportClient->SetViewportType(ELevelViewportType::LVT_Perspective);
@@ -286,9 +354,16 @@ void FUVEditorToolkit::PostInitAssetEditor()
 	// Lit gives us the most options in terms of the materials we can use.
 	ViewportClient->SetViewMode(EViewModeIndex::VMI_Lit);
 
-	// TODO: This should be tied to wherever we store the settings for UV scaling. Right now we
-	// scale [0,1] to [0,1000], hence these settings.
-	ViewportClient->SetViewLocation(FVector(500, 500, 1000));
+	// scale [0,1] to [0,ScaleFactor].
+	// We set our camera to look downward, centered, far enough to be able to see the edges
+	// with a 90 degree FOV
+	double ScaleFactor = 1;
+	UUVEditorSubsystem* UVSubsystem = GEditor->GetEditorSubsystem<UUVEditorSubsystem>();
+	if (UVSubsystem)
+	{
+		ScaleFactor = UUVEditorMode::GetUVMeshScalingFactor();
+	}
+	ViewportClient->SetViewLocation(FVector(ScaleFactor / 2, ScaleFactor/2, ScaleFactor));
 	ViewportClient->SetViewRotation(FRotator(-90, 0, 0));
 
 	// If exposure isn't set to fixed, it will flash as we stare into the void
@@ -299,6 +374,13 @@ void FUVEditorToolkit::PostInitAssetEditor()
 	// We need the viewport client to start out focused, or else it won't get ticked until
 	// we click inside it.
 	ViewportClient->ReceivedFocus(ViewportClient->Viewport);
+
+
+	// Adjust our live preview (3D) viewport
+	// TODO: This should not be hardcoded
+	LivePreviewViewportClient->SetViewLocation(FVector(-200, 100, 100));
+	LivePreviewViewportClient->SetLookAtLocation(FVector(0, 0, 0));
+	LivePreviewViewportClient->ToggleOrbitCamera(true);
 }
 
 #undef LOCTEXT_NAMESPACE
\ No newline at end of file
diff --git a/Engine/Plugins/Experimental/UVEditor/Source/UVEditor/Public/ToolTargets/StaticMeshUVMeshToolTarget.h b/Engine/Plugins/Experimental/UVEditor/Source/UVEditor/Public/ToolTargets/StaticMeshUVMeshToolTarget.h
deleted file mode 100644
diff --git a/Engine/Plugins/Experimental/UVEditor/Source/UVEditor/Public/UVEditorMode.h b/Engine/Plugins/Experimental/UVEditor/Source/UVEditor/Public/UVEditorMode.h
--- a/Engine/Plugins/Experimental/UVEditor/Source/UVEditor/Public/UVEditorMode.h
+++ b/Engine/Plugins/Experimental/UVEditor/Source/UVEditor/Public/UVEditorMode.h
@@ -5,16 +5,19 @@
 #include "CoreMinimal.h"
 
 #include "Tools/UEdMode.h"
+#include "ToolTargets/ToolTarget.h"
 
 #include "UVEditorMode.generated.h"
 
-
+class FToolCommandChange;
 class IUVUnwrapDynamicMesh;
 class UMeshElementsVisualizer;
 class UPreviewMesh;
+class UToolTarget;
 class UWorld;
 class UInteractiveToolPropertySet;
 class UUVToolStateObjectStore;
+class UUVEditorToolMeshInput;
 
 /**
  * The UV editor mode is the mode used in the UV asset editor. It holds most of the inter-tool state.
@@ -31,9 +34,30 @@ public:
 	UUVEditorMode();
 
 	void RegisterTools();
-	void InitializeTargets(TArray<TObjectPtr<UObject>>& UVUnwrapMeshTargetsIn);
 
-	TObjectPtr<UUVToolStateObjectStore> GetStateObjectStore() { return StateObjectStore; }
+	/**
+	 * Gets the tool target requirements for the mode. The resulting targets undergo further processing
+	 * to turn them into the input objects that tools get (since these need preview meshes, etc).
+	 */
+	static const FToolTargetTypeRequirements& GetToolTargetRequirements();
+
+	/**
+	 * Gets the factor by which UV layer unwraps get scaled (scaling makes certain things easier, like zooming in, etc).
+	 */
+	static double GetUVMeshScalingFactor() { return 1000; }
+
+	void InitializeTargets(TArray<TObjectPtr<UObject>>& AssetsIn, TArray<FTransform>* TransformsIn);
+
+	// TODO: We'll probably eventually want a function like this so we can figure out how big our 3d preview
+	// scene is and how we should position the camera intially...
+	//FBoxSphereBounds GetLivePreviewBoundingBox() const;
+
+	// Unlike UInteractiveToolManager::EmitObjectChange, emitting an object change using this
+	// function does not cause it to expire when the active tool doesn't match the emitting tool.
+	// It is important that the emitted change deals properly with expiration itself, for instance
+	// expiring itself when a tool input is invalid or a contained preview is disconnected.
+	// TODO: Should this sort of option exist in UInteractiveToolManager?
+	void EmitToolIndependentObjectChange(UObject* TargetObject, TUniquePtr<FToolCommandChange> Change, const FText& Description);
 
 	// Asset management
 	bool HaveUnappliedChanges();
@@ -66,30 +90,23 @@ protected:
 	TArray<TObjectPtr<UObject>> OriginalObjectsToEdit;
 
 	/**
-	 * This stores the targets generated by the tool target manager that give us the unwrapped UV meshes to work with.
+	 * This stores tool targets created from OriginalObjectsToEdit that provide us with dynamic meshes.
 	 */
 	UPROPERTY()
-	TArray<TScriptInterface<IUVUnwrapDynamicMesh>> UVUnwrapMeshTargets;
+	TArray<TObjectPtr<UToolTarget>> ToolTargets;
 
 	/**
-	 * This is the main storage for the UV meshes, wrapped inside preview objects. Tools will
-	 * get a pointer to these so that they can hide and alter the meshes.
+	 * The actual input objects we give to the tools, including unwrapped meshes and previews, created
+	 * out of ToolTargets.
 	 */
 	UPROPERTY()
-	TArray<TObjectPtr<UPreviewMesh>> DisplayedMeshes;
+	TArray<TObjectPtr<UUVEditorToolMeshInput>> ToolInputObjects;
 
-	/**
-	 * These wireframes track the meshes inside DisplayedMeshes
-	 */
-	UPROPERTY()
-	TArray<TObjectPtr<UMeshElementsVisualizer>> DisplayedWireframes;
 
 	/**
-	 * This lets the tools cache arbitrary state or see if it has been cached by another tool. This
-	 * can store selection, spatial data structures, etc.
+	 * Wireframes have to get ticked to be able to respond to setting changes.
 	 */
-	UPROPERTY()
-	TObjectPtr<UUVToolStateObjectStore> StateObjectStore = nullptr;
+	TArray<TWeakObjectPtr<UMeshElementsVisualizer>> WireframesToTick;
 
 	/**
 	 * If we want to have mode-level property objects, whether user-visible or not,
diff --git a/Engine/Plugins/Experimental/UVEditor/Source/UVEditor/Public/UVEditorSubsystem.h b/Engine/Plugins/Experimental/UVEditor/Source/UVEditor/Public/UVEditorSubsystem.h
--- a/Engine/Plugins/Experimental/UVEditor/Source/UVEditor/Public/UVEditorSubsystem.h
+++ b/Engine/Plugins/Experimental/UVEditor/Source/UVEditor/Public/UVEditorSubsystem.h
@@ -7,8 +7,12 @@
 #include "EditorSubsystem.h"
 #include "ToolTargets/ToolTarget.h" //FToolTargetTypeRequirements
 
+#include "GeometryBase.h"
+
 #include "UVEditorSubsystem.generated.h"
 
+PREDECLARE_GEOMETRY(class FAssetDynamicMeshTargetComboInterface);
+
 class UToolTargetManager;
 class UUVEditor;
 class UToolTarget;
@@ -33,6 +37,13 @@ public:
 	/** Checks that all of the objects are valid targets for a UV editor session. */
 	virtual bool AreObjectsValidTargets(const TArray<UObject*>& InObjects) const;
 
+	/** 
+	 * Tries to build the core targets that provide meshes for UV tools to work on.
+	 */
+	virtual void BuildTargets(const TArray<TObjectPtr<UObject>>& ObjectsIn, 
+		const FToolTargetTypeRequirements& TargetRequirements,
+		TArray<TObjectPtr<UToolTarget>>& TargetsOut);
+
 	/**
 	 * Either brings to the front an existing UV editor instance that is editing one of
 	 * these objects, if one exists, or starts up a new instance editing all of these 
@@ -46,14 +57,11 @@ public:
 	 */
 	virtual void NotifyThatUVEditorClosed(TArray<TObjectPtr<UObject>> ObjectsItWasEditing);
 
-	static const FToolTargetTypeRequirements UVUnwrapMeshTargetRequirements;
-
 protected:
 
 	/**
 	 * Used to let the subsystem figure out whether targets are valid. New factories should be
-	 * added here in Initialize() as they are developed, as well as to the tool manager used
-	 * by a UV mode.
+	 * added here in Initialize() as they are developed.
 	 */
 	UPROPERTY()
 	UToolTargetManager* ToolTargetManager = nullptr;
diff --git a/Engine/Plugins/Experimental/UVEditor/Source/UVEditor/Public/UVEditorToolkit.h b/Engine/Plugins/Experimental/UVEditor/Source/UVEditor/Public/UVEditorToolkit.h
--- a/Engine/Plugins/Experimental/UVEditor/Source/UVEditor/Public/UVEditorToolkit.h
+++ b/Engine/Plugins/Experimental/UVEditor/Source/UVEditor/Public/UVEditorToolkit.h
@@ -7,10 +7,12 @@
 #include "Tools/BaseAssetToolkit.h"
 
 class FAdvancedPreviewScene;
+class FUVEditorPreviewModeManager;
 class IDetailsView;
 class SDockTab;
 class SBorder;
 class UInteractiveToolsContext;
+class UInputRouter;
 
 /**
  * The toolkit is supposed to act as the UI manager for the asset editor. It's responsible 
@@ -27,8 +29,9 @@ public:
 	virtual ~FUVEditorToolkit();
 
 	static const FName InteractiveToolsPanelTabID;
+	static const FName LivePreviewTabID;
 
-	FPreviewScene* GetPreviewScene() { return PreviewScene.Get(); }
+	FPreviewScene* GetPreviewScene() { return UnwrapScene.Get(); }
 
 	// FBaseAssetToolkit
 	virtual void CreateWidgets() override;
@@ -51,6 +54,7 @@ public:
 protected:
 
 	TSharedRef<SDockTab> SpawnTab_InteractiveToolsPanel(const FSpawnTabArgs& Args);
+	TSharedRef<SDockTab> SpawnTab_LivePreview(const FSpawnTabArgs& Args);
 
 	// FBaseAssetToolkit
 	virtual AssetEditorViewportFactoryFunction GetViewportDelegate() override;
@@ -62,5 +66,18 @@ protected:
 	/** Inline content area for the UV mode's content (gotten from UVEditorModeToolkit) */
 	TSharedPtr<SDockTab> ToolsPanel;
 
-	TUniquePtr<FPreviewScene> PreviewScene;
+	/** Scene in which the 2D unwrapped uv meshes live. */
+	TUniquePtr<FPreviewScene> UnwrapScene;
+
+	/** Scene in which the 3D preview meshes of the assets live. */
+	TUniquePtr<FAdvancedPreviewScene> LivePreviewScene;
+
+	// These are related to the 3D "live preview" viewport. The 2d unwrap viewport things are
+	// stored in FBaseAssetToolkit::ViewportTabContent, ViewportDelegate, ViewportClient
+	TSharedPtr<class FEditorViewportTabContent> LivePreviewTabContent;
+	AssetEditorViewportFactoryFunction LivePreviewViewportDelegate;
+	TSharedPtr<FEditorViewportClient> LivePreviewViewportClient;
+	
+	TSharedPtr<FAssetEditorModeManager> LivePreviewEditorModeManager;
+	TObjectPtr<UInputRouter> LivePreviewInputRouter = nullptr;
 };
\ No newline at end of file
diff --git a/Engine/Plugins/Experimental/UVEditor/Source/UVEditor/UVEditor.Build.cs b/Engine/Plugins/Experimental/UVEditor/Source/UVEditor/UVEditor.Build.cs
--- a/Engine/Plugins/Experimental/UVEditor/Source/UVEditor/UVEditor.Build.cs
+++ b/Engine/Plugins/Experimental/UVEditor/Source/UVEditor/UVEditor.Build.cs
@@ -30,6 +30,7 @@ public class UVEditor : ModuleRules
 				"UnrealEd",
 				"InputCore",
 				"InteractiveToolsFramework",
+				"MeshModelingTools",
 				
 				// ... add other public dependencies that you statically link with here ...
 			}
@@ -53,12 +54,11 @@ public class UVEditor : ModuleRules
 				"EditorStyle",
 				"EditorSubsystem",
 				"LevelEditor",
+				"ModelingComponentsEditorOnly", // Static/skeletal mesh tool targets
 				"ModelingComponents",
 				"ToolMenus",
 				"UVEditorTools",
 				
-				// TODO: Remove?
-				"MeshConversion",
 				// ... add private dependencies that you statically link with here ...	
 			}
 			);
diff --git a/Engine/Plugins/Experimental/UVEditor/Source/UVEditorTools/Private/ToolTargets/UVEditorToolMeshInput.cpp b/Engine/Plugins/Experimental/UVEditor/Source/UVEditorTools/Private/ToolTargets/UVEditorToolMeshInput.cpp
new file mode 100644
--- /dev/null
+++ b/Engine/Plugins/Experimental/UVEditor/Source/UVEditorTools/Private/ToolTargets/UVEditorToolMeshInput.cpp
@@ -0,0 +1,369 @@
+// Copyright Epic Games, Inc. All Rights Reserved.
+
+#include "ToolTargets/UVEditorToolMeshInput.h"
+
+#include "Drawing/MeshElementsVisualizer.h" // for wireframe display
+#include "DynamicMesh3.h"
+#include "MeshOpPreviewHelpers.h"
+#include "ModelingToolTargetUtil.h"
+#include "TargetInterfaces/AssetBackedTarget.h"
+#include "TargetInterfaces/DynamicMeshProvider.h"
+#include "TargetInterfaces/MaterialProvider.h"
+#include "ToolSetupUtil.h"
+#include "UVEditorToolUtil.h"
+
+using namespace UE::Geometry;
+
+namespace UVEditorToolMeshInputLocals
+{
+	void CopyMeshPositions(const UE::Geometry::FDynamicMesh3& MeshIn, UE::Geometry::FDynamicMesh3& MeshOut,
+		const TArray<int32>* ChangedVids, const TArray<int32>* ChangedTids)
+	{
+		auto UpdateOrInsertVid = [&MeshIn, &MeshOut](int32 Vid) {
+			if (MeshOut.IsVertex(Vid))
+			{
+				MeshOut.SetVertex(Vid, MeshIn.GetVertex(Vid));
+			}
+			else
+			{
+				MeshOut.InsertVertex(Vid, MeshIn.GetVertex(Vid));
+			}
+		};
+
+		if (!ChangedVids && !ChangedTids)
+		{
+			for (int32 Vid : MeshIn.VertexIndicesItr())
+			{
+				UpdateOrInsertVid(Vid);
+			}
+			for (int32 Tid : MeshIn.TriangleIndicesItr())
+			{
+				MeshOut.SetTriangle(Tid, MeshIn.GetTriangle(Tid));
+			}
+
+			return;
+		}
+
+		if (ChangedVids)
+		{
+			for (int32 Vid : *ChangedVids)
+			{
+				UpdateOrInsertVid(Vid);
+			}
+		}
+
+		if (ChangedTids)
+		{
+			for (int32 Tid : *ChangedTids)
+			{
+				MeshOut.SetTriangle(Tid, MeshIn.GetTriangle(Tid));
+			}
+		}
+	}
+
+	// Aside from potentially  doing partial updates based on ChangedElements and ChangedTids, this
+	// function also differs from a simple overlay copy in that it does not copy the parent pointer arrays.
+	// Also it assumes the existence of all the Tids.
+	void CopyMeshOverlay(const UE::Geometry::FDynamicMeshUVOverlay& OverlayIn, UE::Geometry::FDynamicMeshUVOverlay& OverlayOut,
+		const TArray<int32>* ChangedElements, const TArray<int32>* ChangedTids)
+	{
+		auto UpdateOrInsertElement = [&OverlayIn, &OverlayOut](int32 ElementID) {
+			if (OverlayOut.IsElement(ElementID))
+			{
+				OverlayOut.SetElement(ElementID, OverlayIn.GetElement(ElementID));
+			}
+			else
+			{
+				OverlayOut.InsertElement(ElementID, (float*)OverlayIn.GetElement(ElementID));
+			}
+		};
+
+		if (!ChangedElements && !ChangedTids)
+		{
+			for (int32 ElementID : OverlayIn.ElementIndicesItr())
+			{
+				UpdateOrInsertElement(ElementID);
+			}
+			for (int32 Tid : OverlayIn.GetParentMesh()->TriangleIndicesItr())
+			{
+				OverlayOut.SetTriangle(Tid, OverlayIn.GetTriangle(Tid));
+			}
+
+			return;
+		}
+
+		if (ChangedElements)
+		{
+			for (int32 ElementID : *ChangedElements)
+			{
+				UpdateOrInsertElement(ElementID);
+			}
+		}
+		if (ChangedTids)
+		{
+			for (int32 Tid : *ChangedTids)
+			{
+				OverlayOut.SetTriangle(Tid, OverlayIn.GetTriangle(Tid));
+			}
+		}
+	}
+}
+
+bool UUVEditorToolMeshInput::InitializeMeshes(
+	UToolTarget* Target, int32 UVLayerIndexIn,
+	UWorld* UnwrapWorld, UWorld* LivePreviewWorld,
+	UMaterialInterface* WorkingMaterialIn,
+	TFunction<FVector3d(const FVector2f&)> UVToVertPositionFuncIn,
+	TFunction<FVector2f(const FVector3d&)> VertPositionToUVFuncIn)
+{
+	// TODO: The ModelingToolTargetUtil.h doesn't currently have all the proper functions we want
+	// to access the tool target (for instance, to get a dynamic mesh without a mesh description).
+	// We'll need to update this function once they exist.
+	using namespace UE::ToolTarget;
+
+	checkSlow(Cast<IAssetBackedTarget>(Target) 
+		&& Cast<IDynamicMeshProvider>(Target)
+		&& Cast<IMaterialProvider>(Target));
+
+	OriginalAsset = Cast<IAssetBackedTarget>(Target)->GetSourceData();
+	UVLayerIndex = UVLayerIndexIn;
+	UVToVertPosition = UVToVertPositionFuncIn;
+	VertPositionToUV = VertPositionToUVFuncIn;
+
+	// We are given the preview- i.e. the mesh with the uv layer applied.
+	AppliedCanonical = Cast<IDynamicMeshProvider>(Target)->GetDynamicMesh();
+
+	if (!AppliedCanonical->HasAttributes()
+		|| UVLayerIndex >= AppliedCanonical->Attributes()->NumUVLayers())
+	{
+		return false;
+	}
+
+	// Create a preview version of the 3d mesh
+	AppliedPreview = NewObject<UMeshOpPreviewWithBackgroundCompute>();
+	AppliedPreview->Setup(LivePreviewWorld);
+	AppliedPreview->PreviewMesh->UpdatePreview(AppliedCanonical.Get());
+
+	FComponentMaterialSet MaterialSet;
+	Cast<IMaterialProvider>(Target)->GetMaterialSet(MaterialSet);
+	AppliedPreview->ConfigureMaterials(MaterialSet.Materials, WorkingMaterialIn);
+
+	// Set up the unwrapped mesh
+	UnwrapCanonical = MakeShared<FDynamicMesh3>();
+	UVEditorToolUtil::GenerateUVUnwrapMesh(
+		*AppliedCanonical->Attributes()->GetUVLayer(UVLayerIndex),
+		*UnwrapCanonical, UVToVertPosition);
+
+	// Set up the unwrap preview
+	UnwrapPreview = NewObject<UMeshOpPreviewWithBackgroundCompute>();
+	UnwrapPreview->Setup(UnwrapWorld);
+	UnwrapPreview->PreviewMesh->UpdatePreview(UnwrapCanonical.Get());
+
+	return true;
+}
+
+
+void UUVEditorToolMeshInput::Shutdown()
+{
+	if (WireframeDisplay)
+	{
+		WireframeDisplay->Disconnect();
+		WireframeDisplay = nullptr;
+	}
+
+	UnwrapPreview->Cancel();
+	AppliedPreview->Cancel();
+
+	UnwrapCanonical = nullptr;
+	UnwrapPreview = nullptr;
+	AppliedCanonical = nullptr;
+	AppliedPreview = nullptr;
+
+	OriginalAsset = nullptr;
+}
+
+void UUVEditorToolMeshInput::UpdateUnwrapPreviewOverlay(const TArray<int32>* ChangedVids, const TArray<int32>* ChangedTids)
+{
+	UnwrapPreview->PreviewMesh->DeferredEditMesh([this, ChangedVids, ChangedTids](FDynamicMesh3& Mesh)
+	{
+		FDynamicMeshUVOverlay* DestOverlay = Mesh.Attributes()->PrimaryUV();
+		UVEditorToolUtil::UpdateUVOverlayFromUnwrapMesh(Mesh, *DestOverlay, VertPositionToUV,
+			ChangedVids, ChangedTids);
+	}, false);
+
+	UnwrapPreview->PreviewMesh->NotifyDeferredEditCompleted(UPreviewMesh::ERenderUpdateMode::FastUpdate,
+		EMeshRenderAttributeFlags::Positions | EMeshRenderAttributeFlags::VertexUVs, true); // Assuming user changed positions and didn't notify yet
+
+	if (WireframeDisplay)
+	{
+		WireframeDisplay->NotifyMeshChanged();
+	}
+}
+
+void UUVEditorToolMeshInput::UpdateUnwrapCanonicalOverlay(const TArray<int32>* ChangedVids, const TArray<int32>* ChangedTids)
+{
+	FDynamicMeshUVOverlay* DestOverlay = UnwrapCanonical->Attributes()->PrimaryUV();
+	UVEditorToolUtil::UpdateUVOverlayFromUnwrapMesh(*UnwrapCanonical, *DestOverlay, VertPositionToUV,
+		ChangedVids, ChangedTids);
+}
+
+void UUVEditorToolMeshInput::UpdateAppliedPreviewFromUnwrapPreview(const TArray<int32>* ChangedVids, const TArray<int32>* ChangedTids)
+{
+	using namespace UVEditorToolMeshInputLocals;
+
+	const FDynamicMesh3* SourceUnwrapMesh = UnwrapPreview->PreviewMesh->GetMesh();
+
+	// We need to update the overlay in AppliedPreview. Assuming that the overlay in UnwrapPreview is updated, we can
+	// just copy that overlay (using our function that doesn't copy element parent pointers)
+	AppliedPreview->PreviewMesh->DeferredEditMesh([this, SourceUnwrapMesh, ChangedVids, ChangedTids](FDynamicMesh3& Mesh)
+	{
+			const FDynamicMeshUVOverlay* SourceOverlay = SourceUnwrapMesh->Attributes()->PrimaryUV();
+			FDynamicMeshUVOverlay* DestOverlay = Mesh.Attributes()->PrimaryUV();
+
+			CopyMeshOverlay(*SourceOverlay, *DestOverlay, ChangedVids, ChangedTids);
+	}, false);
+	AppliedPreview->PreviewMesh->NotifyDeferredEditCompleted(UPreviewMesh::ERenderUpdateMode::FastUpdate,
+		EMeshRenderAttributeFlags::VertexUVs, false);
+}
+
+void UUVEditorToolMeshInput::UpdateUnwrapPreviewFromAppliedPreview(const TArray<int32>* ChangedElementIDs, const TArray<int32>* ChangedTids)
+{
+	using namespace UVEditorToolMeshInputLocals;
+
+	// Convert the AppliedPreview UV overlay to positions in UnwrapPreview
+	const FDynamicMeshUVOverlay* SourceOverlay = AppliedPreview->PreviewMesh->GetMesh()->Attributes()->PrimaryUV();
+	UnwrapPreview->PreviewMesh->DeferredEditMesh([this, SourceOverlay, ChangedElementIDs, ChangedTids](FDynamicMesh3& Mesh)
+	{
+		UVEditorToolUtil::UpdateUVUnwrapMesh(*SourceOverlay, Mesh, UVToVertPosition, ChangedElementIDs, ChangedTids);
+
+		// Also copy the actual overlay
+		FDynamicMeshUVOverlay* DestOverlay = Mesh.Attributes()->PrimaryUV();
+		CopyMeshOverlay(*SourceOverlay, *DestOverlay, ChangedElementIDs, ChangedTids);
+	}, false);
+	UnwrapPreview->PreviewMesh->NotifyDeferredEditCompleted(UPreviewMesh::ERenderUpdateMode::FastUpdate,
+		EMeshRenderAttributeFlags::Positions | EMeshRenderAttributeFlags::VertexUVs, true);
+
+	if (WireframeDisplay)
+	{
+		WireframeDisplay->NotifyMeshChanged();
+	}
+}
+
+void UUVEditorToolMeshInput::UpdateCanonicalFromPreviews(const TArray<int32>* ChangedVids, const TArray<int32>* ChangedTids)
+{
+	using namespace UVEditorToolMeshInputLocals;
+
+	// Update UnwrapCanonical from UnwrapPreview
+	UpdateOtherUnwrap(*UnwrapPreview->PreviewMesh->GetMesh(), *UnwrapCanonical, ChangedVids, ChangedTids);
+	
+	// Update the overlay in AppliedCanonical from overlay in AppliedPreview
+	const FDynamicMeshUVOverlay* SourceOverlay = AppliedPreview->PreviewMesh->GetMesh()->Attributes()->GetUVLayer(UVLayerIndex);
+	FDynamicMeshUVOverlay* DestOverlay = AppliedCanonical->Attributes()->GetUVLayer(UVLayerIndex);
+	CopyMeshOverlay(*SourceOverlay, *DestOverlay, ChangedVids, ChangedTids);
+}
+
+void UUVEditorToolMeshInput::UpdateAllFromUnwrapPreview(
+	const TArray<int32>* ChangedVids, const TArray<int32>* ChangedTids)
+{
+	const FDynamicMesh3* ReferenceMesh = UnwrapPreview->PreviewMesh->GetMesh();
+
+	// Update UnwrapCanonical
+	UpdateOtherUnwrap(*ReferenceMesh, *UnwrapCanonical, ChangedVids, ChangedTids);
+
+	// Update the applied meshes
+	const FDynamicMeshUVOverlay* SourceOverlay = ReferenceMesh->Attributes()->PrimaryUV();
+	UpdateAppliedOverlays(*SourceOverlay, ChangedVids, ChangedTids);
+}
+
+void UUVEditorToolMeshInput::UpdateAllFromUnwrapCanonical(
+	const TArray<int32>* ChangedVids, const TArray<int32>* ChangedTids)
+{
+	// Update UnwrapPreview
+	UnwrapPreview->PreviewMesh->DeferredEditMesh([this, ChangedVids, ChangedTids](FDynamicMesh3& Mesh)
+	{
+		UpdateOtherUnwrap(*UnwrapCanonical, Mesh, ChangedVids, ChangedTids);
+	}, false);
+	UnwrapPreview->PreviewMesh->NotifyDeferredEditCompleted(UPreviewMesh::ERenderUpdateMode::FastUpdate,
+		EMeshRenderAttributeFlags::Positions | EMeshRenderAttributeFlags::VertexUVs, true);
+	
+
+	// Update the applied meshes
+	FDynamicMeshUVOverlay* SourceOverlay = UnwrapCanonical->Attributes()->PrimaryUV();
+	UpdateAppliedOverlays(*SourceOverlay, ChangedVids, ChangedTids);
+
+	if (WireframeDisplay)
+	{
+		WireframeDisplay->NotifyMeshChanged();
+	}
+}
+
+void UUVEditorToolMeshInput::UpdateAllFromAppliedPreview(
+	const TArray<int32>* ChangedElementIDs, const TArray<int32>* ChangedTids)
+{
+	using namespace UVEditorToolMeshInputLocals;
+
+	// Update AppliedCanonical
+	const FDynamicMeshUVOverlay* SourceOverlay = AppliedPreview->PreviewMesh->GetMesh()
+		->Attributes()->GetUVLayer(UVLayerIndex);
+	FDynamicMeshUVOverlay* DestOverlay = AppliedCanonical->Attributes()->GetUVLayer(UVLayerIndex);
+	CopyMeshOverlay(*SourceOverlay, *DestOverlay, ChangedElementIDs, ChangedTids);
+
+	// Update UnwrapCanonical
+	UVEditorToolUtil::UpdateUVUnwrapMesh(*SourceOverlay, *UnwrapCanonical, 
+		UVToVertPosition, ChangedElementIDs, ChangedTids);
+
+	// Update UnwrapPreview from UnwrapCanonical
+	UnwrapPreview->PreviewMesh->DeferredEditMesh([this, ChangedElementIDs, ChangedTids](FDynamicMesh3& Mesh) {
+		UpdateOtherUnwrap(*UnwrapCanonical, Mesh, ChangedElementIDs, ChangedTids);
+	}, false);
+	UnwrapPreview->PreviewMesh->NotifyDeferredEditCompleted(UPreviewMesh::ERenderUpdateMode::FastUpdate,
+		EMeshRenderAttributeFlags::Positions | EMeshRenderAttributeFlags::VertexUVs, true);
+
+	if (WireframeDisplay)
+	{
+		WireframeDisplay->NotifyMeshChanged();
+	}
+}
+
+/**
+ * Helper function. Updates the UV overlays of AppliedCanonical and AppliedPreview.
+ */
+void UUVEditorToolMeshInput::UpdateAppliedOverlays(const FDynamicMeshUVOverlay& SourceOverlay,
+	const TArray<int32>* ChangedVids, const TArray<int32>* ChangedTids)
+{
+	using namespace UVEditorToolMeshInputLocals;
+
+	// Update UnwrappedMeshCanonical
+	FDynamicMeshUVOverlay* DestOverlay = AppliedCanonical->Attributes()->GetUVLayer(UVLayerIndex);
+	CopyMeshOverlay(SourceOverlay, *DestOverlay, ChangedVids, ChangedTids);
+
+	// Update AppliedPreview
+	AppliedPreview->PreviewMesh->DeferredEditMesh([this, &SourceOverlay, ChangedVids, ChangedTids](FDynamicMesh3& Mesh)
+	{
+		FDynamicMeshUVOverlay* DestOverlay = Mesh.Attributes()->GetUVLayer(UVLayerIndex);
+		CopyMeshOverlay(SourceOverlay, *DestOverlay, ChangedVids, ChangedTids);
+	}, false);
+	AppliedPreview->PreviewMesh->NotifyDeferredEditCompleted(UPreviewMesh::ERenderUpdateMode::FastUpdate,
+		EMeshRenderAttributeFlags::VertexUVs, false);
+}
+
+/**
+ * Helper function. Uses the positions and UV overlay of one unwrap mesh to update another one.
+ */
+void UUVEditorToolMeshInput::UpdateOtherUnwrap(const FDynamicMesh3& SourceUnwrapMesh,
+	FDynamicMesh3& DestUnwrapMesh, const TArray<int32>* ChangedVids, const TArray<int32>* ChangedTids)
+{
+	using namespace UVEditorToolMeshInputLocals;
+
+	if (!ChangedVids && !ChangedTids)
+	{
+		DestUnwrapMesh.Copy(SourceUnwrapMesh, false, false, false, true); // Copy positions and UVs
+		return;
+	}
+
+	CopyMeshPositions(SourceUnwrapMesh, DestUnwrapMesh, ChangedVids, ChangedTids);
+
+	const FDynamicMeshUVOverlay* SourceOverlay = SourceUnwrapMesh.Attributes()->PrimaryUV();
+	FDynamicMeshUVOverlay* DestOverlay = DestUnwrapMesh.Attributes()->PrimaryUV();
+	CopyMeshOverlay(*SourceOverlay, *DestOverlay, ChangedVids, ChangedTids);
+}
diff --git a/Engine/Plugins/Experimental/UVEditor/Source/UVEditorTools/Private/UVEditorToolUtil.cpp b/Engine/Plugins/Experimental/UVEditor/Source/UVEditorTools/Private/UVEditorToolUtil.cpp
new file mode 100644
--- /dev/null
+++ b/Engine/Plugins/Experimental/UVEditor/Source/UVEditorTools/Private/UVEditorToolUtil.cpp
@@ -0,0 +1,151 @@
+// Copyright Epic Games, Inc. All Rights Reserved.
+
+#include "UVEditorToolUtil.h"
+
+#include "DynamicMesh3.h"
+#include "DynamicMeshAttributeSet.h"
+
+using namespace UE::Geometry;
+
+void UVEditorToolUtil::GenerateUVUnwrapMesh(
+	const FDynamicMeshUVOverlay& UVOverlay, FDynamicMesh3& UnwrapMeshOut,
+	TFunctionRef<FVector3d(const FVector2f&)> UVToVertPosition)
+{
+	UnwrapMeshOut.Clear();
+
+	// The unwrap mesh will have an overlay on top of it with the corresponding UVs,
+	// in case we want to draw the texture on it, etc. However note that we can't
+	// just do a Copy() call using the source overlay because the parent vertices will differ.
+	UnwrapMeshOut.EnableAttributes(); // Makes one UV layer
+	FDynamicMeshUVOverlay* UnwrapMeshUVOverlay = UnwrapMeshOut.Attributes()->PrimaryUV();
+
+	// Create a vert for each uv overlay element
+	UnwrapMeshOut.BeginUnsafeVerticesInsert();
+	UnwrapMeshUVOverlay->BeginUnsafeElementsInsert();
+	for (int32 ElementID : UVOverlay.ElementIndicesItr())
+	{
+		FVector2f UVElement = UVOverlay.GetElement(ElementID);
+		UnwrapMeshOut.InsertVertex(ElementID, UVToVertPosition(UVElement), true);
+		UnwrapMeshUVOverlay->InsertElement(ElementID, (float*)UVElement);
+	}
+	UnwrapMeshOut.EndUnsafeVerticesInsert();
+	UnwrapMeshUVOverlay->EndUnsafeElementsInsert();
+
+	// Insert a tri connecting the same vids as elements in the overlay.
+	const FDynamicMesh3* OverlayParentMesh = UVOverlay.GetParentMesh();
+	UnwrapMeshOut.BeginUnsafeTrianglesInsert();
+	for (int32 Tid : OverlayParentMesh->TriangleIndicesItr())
+	{
+		if (UVOverlay.IsSetTriangle(Tid))
+		{
+			FIndex3i UVTri = UVOverlay.GetTriangle(Tid);
+			UnwrapMeshOut.InsertTriangle(Tid, UVTri, 0, true);
+			UnwrapMeshUVOverlay->SetTriangle(Tid, UVTri);
+		}
+	}
+	UnwrapMeshOut.EndUnsafeTrianglesInsert();
+}
+
+void UVEditorToolUtil::UpdateUVUnwrapMesh(const FDynamicMeshUVOverlay& UVOverlayIn, 
+	FDynamicMesh3& UnwrapMeshOut, TFunctionRef<FVector3d(const FVector2f&)> UVToVertPosition,
+	const TArray<int32>* ChangedElementIDs, const TArray<int32>* ChangedTids)
+{
+	if (!ChangedElementIDs && !ChangedTids)
+	{
+		UVEditorToolUtil::GenerateUVUnwrapMesh(UVOverlayIn, UnwrapMeshOut, UVToVertPosition);
+		return;
+	}
+
+	// Otherwise, we only update the parts that we're asked.
+	FDynamicMeshUVOverlay* UnwrapMeshUVOverlay = UnwrapMeshOut.Attributes()->GetUVLayer(0);
+	if (ChangedElementIDs)
+	{
+		for (int32 ElementID : *ChangedElementIDs)
+		{
+			FVector2f ElementValue = UVOverlayIn.GetElement(ElementID);
+			
+			// Update the actual unwrap mesh
+			if (UnwrapMeshOut.IsVertex(ElementID))
+			{
+				UnwrapMeshOut.SetVertex(ElementID, UVToVertPosition(ElementValue));
+			}
+			else
+			{
+				UnwrapMeshOut.InsertVertex(ElementID, UVToVertPosition(ElementValue));
+			}
+
+			// Update the unwrap overlay
+			if (UnwrapMeshUVOverlay->IsElement(ElementID))
+			{
+				UnwrapMeshUVOverlay->SetElement(ElementID, ElementValue);
+			}
+			else
+			{
+				UnwrapMeshUVOverlay->InsertElement(ElementID, (float*)ElementValue);
+			}
+		}
+	}
+
+	if (ChangedTids)
+	{
+		for (int32 Tid : *ChangedTids)
+		{
+			FIndex3i Triangle = UVOverlayIn.GetTriangle(Tid);
+			UnwrapMeshUVOverlay->SetTriangle(Tid, Triangle);
+			UnwrapMeshOut.SetTriangle(Tid, Triangle);
+		}
+	}
+}
+
+void UVEditorToolUtil::UpdateUVOverlayFromUnwrapMesh(
+	const FDynamicMesh3& UnwrapMeshIn, FDynamicMeshUVOverlay& UVOverlayOut,
+	TFunctionRef<FVector2f(const FVector3d&)> VertPositionToUV,
+	const TArray<int32>* ChangedVids, const TArray<int32>* ChangedTids)
+{
+	if (!ensure(UVOverlayOut.GetParentMesh()->MaxTriangleID() == UnwrapMeshIn.MaxTriangleID()))
+	{
+		return;
+	}
+
+	auto UpdateUVElement = [&UnwrapMeshIn, &UVOverlayOut, VertPositionToUV](int32 Vid) {
+		if (UVOverlayOut.IsElement(Vid))
+		{
+			UVOverlayOut.SetElement(Vid, VertPositionToUV(UnwrapMeshIn.GetVertex(Vid)));
+		}
+		else
+		{
+			UVOverlayOut.InsertElement(Vid, (float*)VertPositionToUV(UnwrapMeshIn.GetVertex(Vid)));
+		}
+	};
+
+	// Update all vids and tids if not given changed items
+	if (!ChangedVids && !ChangedTids)
+	{
+		for (int32 Vid : UnwrapMeshIn.VertexIndicesItr())
+		{
+			UpdateUVElement(Vid);
+		}
+		for (int32 Tid : UnwrapMeshIn.TriangleIndicesItr())
+		{
+			UVOverlayOut.SetTriangle(Tid, UnwrapMeshIn.GetTriangle(Tid));
+		}
+
+		return;
+	}
+
+	//Otherwise, only update what we're told about.
+	if (ChangedVids)
+	{
+		for (int32 Vid : *ChangedVids)
+		{
+			UpdateUVElement(Vid);
+		}
+	}
+	if (ChangedTids)
+	{
+		for (int32 Tid : *ChangedTids)
+		{
+			UVOverlayOut.SetTriangle(Tid, UnwrapMeshIn.GetTriangle(Tid));
+		}
+	}
+}
diff --git a/Engine/Plugins/Experimental/UVEditor/Source/UVEditorTools/Private/UVSelectTool.cpp b/Engine/Plugins/Experimental/UVEditor/Source/UVEditorTools/Private/UVSelectTool.cpp
--- a/Engine/Plugins/Experimental/UVEditor/Source/UVEditorTools/Private/UVSelectTool.cpp
+++ b/Engine/Plugins/Experimental/UVEditor/Source/UVEditorTools/Private/UVSelectTool.cpp
@@ -3,12 +3,19 @@
 #include "UVSelectTool.h"
 
 #include "BaseGizmos/TransformGizmo.h"
-#include "TargetInterfaces/UVUnwrapDynamicMesh.h"
+#include "ContextObjectStore.h"
+#include "Drawing/LineSetComponent.h"
+#include "Drawing/MeshElementsVisualizer.h"
+#include "Drawing/PreviewGeometryActor.h"
+#include "DynamicMeshChangeTracker.h"
+#include "ToolTargets/UVEditorToolMeshInput.h"
 #include "InteractiveToolManager.h"
+#include "MeshOpPreviewHelpers.h" // UMeshOpPreviewWithBackgroundCompute
 #include "PreviewMesh.h"
 #include "Selection/MeshSelectionMechanic.h"
 #include "Selection/DynamicMeshSelection.h"
-#include "UVToolStateObjects.h"
+#include "ToolSetupUtil.h"
+#include "UVToolContextObjects.h"
 
 #include "ToolTargetManager.h"
 
@@ -20,38 +27,30 @@ using namespace UE::Geometry;
  * ToolBuilder
  */
 
-
-// TODO: We may need to rework how targets are passed to this tool. Right now GetTargetRequirements
-// is not used, and is not even accurate, since we don't even use UUVUnwrapDynamicMesh ourselves.
-// Instead, the UVMode is the one that uses UUVUnwrapDynamicMesh to create meshes that it can display,
-// and we just get those preview meshes passed to us. Perhaps we should get passed some sort of
-// Target-wrapped preview mesh?
-// The other option is to use UUVUnwrapDynamicMesh in the tool and cache it for use by others, but
-// this ends up being a lot less clean to code up, and it would be preferable to avoid that messiness.
 const FToolTargetTypeRequirements& UUVSelectToolBuilder::GetTargetRequirements() const
 {
-	static FToolTargetTypeRequirements TypeRequirements(UUVUnwrapDynamicMesh::StaticClass());
+	static FToolTargetTypeRequirements TypeRequirements(UUVEditorToolMeshInput::StaticClass());
 	return TypeRequirements;
 }
 
 bool UUVSelectToolBuilder::CanBuildTool(const FToolBuilderState& SceneState) const
 {
-	return DisplayedMeshes->Num() > 0;
+	return Targets && Targets->Num() > 0;
 }
 
 UInteractiveTool* UUVSelectToolBuilder::BuildTool(const FToolBuilderState& SceneState) const
 {
 	UUVSelectTool* NewTool = NewObject<UUVSelectTool>(SceneState.ToolManager);
+	NewTool->SetWorld(SceneState.World);
 	NewTool->SetGizmoEnabled(bGizmoEnabled);
-	NewTool->SetStateObjectStore(StateObjectStore);
-	NewTool->SetDisplayedMeshes(*DisplayedMeshes);
+	NewTool->SetTargets(*Targets);
 
 	return NewTool;
 }
 
 void UUVSelectTool::Setup()
 {
-	check(DisplayedMeshes.Num() > 0);
+	check(Targets.Num() > 0);
 
 	UInteractiveTool::Setup();
 
@@ -63,53 +62,47 @@ void UUVSelectTool::Setup()
 
 	SelectionMechanic = NewObject<UMeshSelectionMechanic>();
 	SelectionMechanic->Setup(this);
+	SelectionMechanic->SetWorld(Targets[0]->UnwrapPreview->GetWorld());
 	SelectionMechanic->OnSelectionChanged.AddUObject(this, &UUVSelectTool::OnSelectionChanged);
 
 	SelectionMechanic->SelectionMode =
 		Settings->SelectionMode == EUVSelectToolSelectionMode::Island ? EMeshSelectionMechanicMode::Component
 		: EMeshSelectionMechanicMode::Edge;
 
-
-	// TODO: Should StateObjectStore be mandatory? We could enforce it in builder.
-	// Restore any saved state
-	if (StateObjectStore)
+	// Retrieve cached AABB tree storage, or else set it up
+	UContextObjectStore* ContextStore = GetToolManager()->GetContextObjectStore();
+	UUVToolAABBTreeStorage* TreeStore = ContextStore->FindContext<UUVToolAABBTreeStorage>();
+	if (!TreeStore)
 	{
-		UUVMeshAABBTrees* TreeStore = StateObjectStore->GetToolStateObject<UUVMeshAABBTrees>();
-		if (TreeStore)
-		{
-			AABBTrees = TreeStore->AABBTrees;
-		}
-			
-		UUVMeshSelection* SelectionStore = StateObjectStore->GetToolStateObject<UUVMeshSelection>();
-		if (SelectionStore)
-		{
-			SelectionMechanic->SetSelection(*SelectionStore->Selection);
-		}
+		TreeStore = NewObject<UUVToolAABBTreeStorage>();
+		ContextStore->AddContextObject(TreeStore);
 	}
 
-	// Initialize the spatial structures if they weren't there to load
-	if (AABBTrees.IsEmpty())
+	// Initialize the AABB trees from cached values, or make new ones.
+	for (TObjectPtr<UUVEditorToolMeshInput> Target : Targets)
 	{
-		for (TObjectPtr<UPreviewMesh> Mesh : DisplayedMeshes)
-		{
-			TSharedPtr<FDynamicMeshAABBTree3> Tree = MakeShared<FDynamicMeshAABBTree3>();
-			Tree->SetMesh(Mesh->GetMesh());
-			AABBTrees.Add(Tree);
-		}
-
-		// Save them if appropriate
-		if (StateObjectStore)
+		TSharedPtr<FDynamicMeshAABBTree3> Tree = TreeStore->Get(Target->UnwrapCanonical.Get());
+		if (!Tree)
 		{
-			UUVMeshAABBTrees* TreeStore = NewObject<UUVMeshAABBTrees>();
-			TreeStore->AABBTrees = AABBTrees;
-			StateObjectStore->SetToolStateObject(TreeStore);
+			Tree = MakeShared<FDynamicMeshAABBTree3>();
+			Tree->SetMesh(Target->UnwrapCanonical.Get());
+			TreeStore->Set(Target->UnwrapCanonical.Get(), Tree);
 		}
+		AABBTrees.Add(Tree);
 	}
 
 	// Add the spatial structures to the selection mechanic
-	for (TSharedPtr<FDynamicMeshAABBTree3> Tree : AABBTrees)
+	for (int32 i = 0; i < Targets.Num(); ++i)
 	{
-		SelectionMechanic->AddSpatial(Tree);
+		SelectionMechanic->AddSpatial(AABBTrees[i],
+			Targets[i]->UnwrapPreview->PreviewMesh->GetTransform());
+	}
+
+	// See if we have a stored selection
+	UUVToolMeshSelection* SelectionStore = ContextStore->FindContext<UUVToolMeshSelection>();
+	if (SelectionStore)
+	{
+		SelectionMechanic->SetSelection(*SelectionStore->Selection);
 	}
 
 	// Gizmo setup
@@ -117,7 +110,7 @@ void UUVSelectTool::Setup()
 	UTransformProxy* TransformProxy = NewObject<UTransformProxy>(this);
 	TransformGizmo = GizmoManager->CreateCustomTransformGizmo(
 		ETransformGizmoSubElements::TranslateAxisX | ETransformGizmoSubElements::TranslateAxisY | ETransformGizmoSubElements::TranslatePlaneXY
-		//| ETransformGizmoSubElements::ScaleAxisX | ETransformGizmoSubElements::ScaleAxisY | ETransformGizmoSubElements::ScalePlaneXY
+		| ETransformGizmoSubElements::ScaleAxisX | ETransformGizmoSubElements::ScaleAxisY | ETransformGizmoSubElements::ScalePlaneXY
 		| ETransformGizmoSubElements::RotateAxisZ,
 		this);
 	TransformProxy->OnBeginTransformEdit.AddUObject(this, &UUVSelectTool::GizmoTransformStarted);
@@ -128,6 +121,14 @@ void UUVSelectTool::Setup()
 	TransformGizmo->bUseContextCoordinateSystem = false;
 	TransformGizmo->SetActiveTarget(TransformProxy, GetToolManager());
 
+	LivePreviewGeometryActor = Targets[0]->AppliedPreview->GetWorld()->SpawnActor<APreviewGeometryActor>(
+		FVector::ZeroVector, FRotator(0, 0, 0), FActorSpawnParameters());
+	LivePreviewLineSet = NewObject<ULineSetComponent>(LivePreviewGeometryActor);
+	LivePreviewGeometryActor->SetRootComponent(LivePreviewLineSet);
+	LivePreviewLineSet->RegisterComponent();
+	LivePreviewLineSet->SetLineMaterial(ToolSetupUtil::GetDefaultLineComponentMaterial(
+		GetToolManager(), /*bDepthTested*/ true));
+
 	if (!SelectionMechanic->GetCurrentSelection().IsEmpty())
 	{
 		OnSelectionChanged();
@@ -137,18 +138,32 @@ void UUVSelectTool::Setup()
 
 void UUVSelectTool::Shutdown(EToolShutdownType ShutdownType)
 {
-	if (StateObjectStore)
+	UContextObjectStore* ContextStore = GetToolManager()->GetContextObjectStore();
+	if (!SelectionMechanic->GetCurrentSelection().IsEmpty())
 	{
-		UUVMeshSelection* SelectionStore = NewObject<UUVMeshSelection>();
-		SelectionStore->Selection = MakeShared<FDynamicMeshSelection>();
+		UUVToolMeshSelection* SelectionStore = ContextStore->FindContext<UUVToolMeshSelection>();
+		if (!SelectionStore)
+		{
+			SelectionStore = NewObject<UUVToolMeshSelection>();
+			ContextStore->AddContextObject(SelectionStore);
+		}
 		*SelectionStore->Selection = SelectionMechanic->GetCurrentSelection();
-		StateObjectStore->SetToolStateObject(SelectionStore);
+	}
+	else
+	{
+		ContextStore->RemoveContextObjectsOfType<UUVToolMeshSelection>();
 	}
 
 	Settings->SaveProperties(this);
 
 	SelectionMechanic->Shutdown();
 
+	if (LivePreviewGeometryActor)
+	{
+		LivePreviewGeometryActor->Destroy();
+		LivePreviewGeometryActor = nullptr;
+	}
+
 	// Calls shutdown on gizmo and destroys it.
 	GetToolManager()->GetPairedGizmoManager()->DestroyAllGizmosByOwner(this);
 }
@@ -166,12 +181,7 @@ void UUVSelectTool::UpdateGizmo()
 
 	if (!Selection.IsEmpty())
 	{
-		FVector3d Centroid = FVector3d::Zero();
-		for (int32 Vid : MovingVids)
-		{
-			Centroid += Selection.Mesh->GetVertex(Vid);
-		}
-		Centroid /= MovingVids.Num();
+		FVector3d Centroid = SelectionMechanic->GetCurrentSelectionCentroid();
 
 		// TODO: We should either make it so that transparent materials don't affect the rendering
 		// of the gizmo component, or else we should pass in some depth offset into the component.
@@ -188,25 +198,27 @@ void UUVSelectTool::OnSelectionChanged()
 {
 	const FDynamicMeshSelection& Selection = SelectionMechanic->GetCurrentSelection();
 
-	SelectionMeshIndex = -1;
+	SelectionTargetIndex = -1;
 	MovingVids.Reset();
+	BoundaryEids.Reset();
 	if (!Selection.IsEmpty())
 	{
 		// Note which mesh we're selecting in.
-		for (int32 i = 0; i < DisplayedMeshes.Num(); ++i)
+		for (int32 i = 0; i < Targets.Num(); ++i)
 		{
-			if (DisplayedMeshes[i]->GetMesh() == Selection.Mesh)
+			if (Targets[i]->UnwrapCanonical.Get() == Selection.Mesh)
 			{
-				SelectionMeshIndex = i;
+				SelectionTargetIndex = i;
 				break;
 			}
 		}
-		check(SelectionMeshIndex >= 0);
+		check(SelectionTargetIndex >= 0);
 
 		// Note the selected vids
 		TSet<int32> VidSet;
 		if (Selection.Type == FDynamicMeshSelection::EType::Triangle)
 		{
+			const FDynamicMesh3* LivePreviewMesh = Targets[SelectionTargetIndex]->AppliedCanonical.Get();
 			for (int32 Tid : Selection.SelectedIDs)
 			{
 				FIndex3i TriVids = Selection.Mesh->GetTriangle(Tid);
@@ -218,6 +230,21 @@ void UUVSelectTool::OnSelectionChanged()
 						MovingVids.Add(TriVids[i]);
 					}
 				}
+
+				// Gather the boundary edges in the live preview
+				FIndex3i TriEids = LivePreviewMesh->GetTriEdges(Tid);
+				for (int i = 0; i < 3; ++i)
+				{
+					FIndex2i EdgeTids = LivePreviewMesh->GetEdgeT(TriEids[i]);
+					for (int j = 0; j < 2; ++j)
+					{
+						if (EdgeTids[j] != Tid && !Selection.SelectedIDs.Contains(EdgeTids[j]))
+						{
+							BoundaryEids.Add(TriEids[i]);
+							break;
+						}
+					}
+				}
 			}
 		}
 		else if (Selection.Type == FDynamicMeshSelection::EType::Edge)
@@ -240,10 +267,35 @@ void UUVSelectTool::OnSelectionChanged()
 			check(false);
 		}
 	}
+
+	UpdateLivePreviewLines();
 	
 	UpdateGizmo();
 }
 
+void UUVSelectTool::UpdateLivePreviewLines()
+{
+	LivePreviewLineSet->Clear();
+
+	const FDynamicMeshSelection& Selection = SelectionMechanic->GetCurrentSelection();
+	if (!Selection.IsEmpty())
+	{
+		FTransform MeshTransform = Targets[SelectionTargetIndex]->AppliedPreview->PreviewMesh->GetTransform();
+		const FDynamicMesh3* LivePreviewMesh = Targets[SelectionTargetIndex]->AppliedCanonical.Get();
+
+		for (int32 Eid : BoundaryEids)
+		{
+			FVector3d Vert1, Vert2;
+			LivePreviewMesh->GetEdgeV(Eid, Vert1, Vert2);
+
+			LivePreviewLineSet->AddLine(
+				MeshTransform.TransformPosition(Vert1), 
+				MeshTransform.TransformPosition(Vert2), 
+				FColor::Yellow, 2, 1.5);
+		}
+	}
+}
+
 void UUVSelectTool::SetGizmoEnabled(bool bEnabledIn)
 {
 	bGizmoEnabled = bEnabledIn;
@@ -260,7 +312,9 @@ void UUVSelectTool::GizmoTransformStarted(UTransformProxy* Proxy)
 {
 	InitialGizmoFrame = FFrame3d(TransformGizmo->ActiveTarget->GetTransform());
 	MovingVertOriginalPositions.SetNum(MovingVids.Num());
-	const FDynamicMesh3* Mesh = DisplayedMeshes[SelectionMeshIndex]->GetMesh();
+	const FDynamicMesh3* Mesh = Targets[SelectionTargetIndex]->UnwrapCanonical.Get();
+	// Note: Our meshes currently don't have a transform. Otherwise we'd need to convert vid location to world
+	// space first, then to the frame.
 	for (int32 i = 0; i < MovingVids.Num(); ++i)
 	{
 		MovingVertOriginalPositions[i] = InitialGizmoFrame.ToFramePoint(Mesh->GetVertex(MovingVids[i]));
@@ -269,32 +323,64 @@ void UUVSelectTool::GizmoTransformStarted(UTransformProxy* Proxy)
 
 void UUVSelectTool::GizmoTransformChanged(UTransformProxy* Proxy, FTransform Transform)
 {
-	FFrame3d NewGizmoFrame(TransformGizmo->ActiveTarget->GetTransform());
-	FTransform DeltaTransform = NewGizmoFrame.ToFTransform().GetRelativeTransform(InitialGizmoFrame.ToFTransform());
-	const FDynamicMeshSelection& Selection = SelectionMechanic->GetCurrentSelection();
+	FTransform DeltaTransform = Transform.GetRelativeTransform(InitialGizmoFrame.ToFTransform());
 
-	if (!DeltaTransform.GetTranslation().IsNearlyZero() || !DeltaTransform.GetRotation().IsIdentity())
+	if (!DeltaTransform.GetTranslation().IsNearlyZero() || !DeltaTransform.GetRotation().IsIdentity() || Transform.GetScale3D() != FVector::One())
 	{
-		DisplayedMeshes[SelectionMeshIndex]->EditMesh([&Selection, &NewGizmoFrame, this](FDynamicMesh3& MeshIn)
+		FTransform3d TransformToApply(Transform);
+
+		// TODO: The division here is a bit of a hack. Properly-speaking, the scaling handles should act relative to
+		// gizmo size, not the visible space across which we drag, otherwise it becomes dependent on the units we
+		// use and our absolute distance from the object. Since our UV unwrap is scaled by 1000 to make it
+		// easier to zoom in and out without running into issues, the measure of the distance across which we typically
+		// drag the handles is too high to be convenient. Until we make the scaling invariant to units/distance from
+		// target, we use this hack.
+		TransformToApply.SetScale(FVector::One() + (Transform.GetScale3D() - FVector::One()) / 10);
+
+		const FDynamicMeshSelection& Selection = SelectionMechanic->GetCurrentSelection();
+
+		Targets[SelectionTargetIndex]->UnwrapPreview->PreviewMesh->DeferredEditMesh([&TransformToApply,  &Transform, this](FDynamicMesh3& MeshIn)
 		{
 			for (int32 i = 0; i < MovingVids.Num(); ++i)
 			{
-				MeshIn.SetVertex(MovingVids[i], NewGizmoFrame.FromFramePoint(MovingVertOriginalPositions[i]));
+				MeshIn.SetVertex(MovingVids[i], TransformToApply.TransformPosition(MovingVertOriginalPositions[i]));
 			}
-		});
+		}, false);
+		Targets[SelectionTargetIndex]->UpdateUnwrapPreviewOverlay(&MovingVids);
 
-		// TODO: This doesn't get broadcast when EditMesh is called. Should it? Or should there be some other way
-		// to inform the wireframe that the mesh changed?
-		DisplayedMeshes[SelectionMeshIndex]->GetOnMeshChanged().Broadcast();
+		SelectionMechanic->SetDrawnElementsTransform((FTransform)TransformToApply);
+
+		if (Settings->bUpdatePreviewDuringDrag)
+		{
+			Targets[SelectionTargetIndex]->UpdateAppliedPreviewFromUnwrapPreview(&MovingVids);
+		}
 	}	
 }
 
 void UUVSelectTool::GizmoTransformEnded(UTransformProxy* Proxy)
 {
-	if (!AABBTrees[SelectionMeshIndex]->IsValid())
+	// TODO: Add undo support
+	//FDynamicMeshChangeTracker ChangeTracker(Targets[SelectionTargetIndex]->UnwrapPreview->PreviewMesh->GetMesh());
+	//ChangeTracker.BeginChange();
+	//ChangeTracker.SaveTriangles(MovingVids, true);
+
+	if (Settings->bUpdatePreviewDuringDrag)
+	{
+		// Both previews must already be updated, so only need to update canonical
+		Targets[SelectionTargetIndex]->UpdateCanonicalFromPreviews(&MovingVids);
+	}
+	else
 	{
-		AABBTrees[SelectionMeshIndex]->Build();
+		Targets[SelectionTargetIndex]->UpdateAllFromUnwrapPreview(&MovingVids);
 	}
+
+	if (!AABBTrees[SelectionTargetIndex]->IsValid())
+	{
+		AABBTrees[SelectionTargetIndex]->Build();
+	}
+
+	TransformGizmo->SetNewChildScale(FVector::One());
+	SelectionMechanic->RebuildDrawnElements(TransformGizmo->GetGizmoTransform());
 }
 
 void UUVSelectTool::Render(IToolsContextRenderAPI* RenderAPI)
diff --git a/Engine/Plugins/Experimental/UVEditor/Source/UVEditorTools/Private/UVToolContextObjects.cpp b/Engine/Plugins/Experimental/UVEditor/Source/UVEditorTools/Private/UVToolContextObjects.cpp
new file mode 100644
--- /dev/null
+++ b/Engine/Plugins/Experimental/UVEditor/Source/UVEditorTools/Private/UVToolContextObjects.cpp
@@ -0,0 +1,48 @@
+// Copyright Epic Games, Inc. All Rights Reserved.
+
+#pragma once
+
+#include "UVToolContextObjects.h"
+
+#include "InputRouter.h"
+
+using namespace UE::Geometry;
+
+void UUVToolLivePreviewAPI::Initialize(UWorld* WorldIn, UInputRouter* RouterIn)
+{
+	World = WorldIn;
+	InputRouter = RouterIn;
+}
+
+void UUVToolAABBTreeStorage::Set(FDynamicMesh3* MeshKey, TSharedPtr<FDynamicMeshAABBTree3> Tree)
+{
+	AABBTrees.Add(MeshKey, Tree);
+}
+
+TSharedPtr<FDynamicMeshAABBTree3> UUVToolAABBTreeStorage::Get(FDynamicMesh3* MeshKey)
+{
+	TSharedPtr<FDynamicMeshAABBTree3>* FoundPtr = AABBTrees.Find(MeshKey);
+	return FoundPtr ? *FoundPtr : nullptr;
+}
+
+void UUVToolAABBTreeStorage::Remove(FDynamicMesh3* MeshKey)
+{
+	AABBTrees.Remove(MeshKey);
+}
+
+void UUVToolAABBTreeStorage::RemoveByPredicate(TUniqueFunction<
+	bool(const TPair<FDynamicMesh3*, TSharedPtr<FDynamicMeshAABBTree3>>&)> Predicate)
+{
+	for (auto It = AABBTrees.CreateIterator(); It; ++It)
+	{
+		if (Predicate(*It))
+		{
+			It.RemoveCurrent();
+		}
+	}
+}
+
+void UUVToolAABBTreeStorage::Empty()
+{
+	AABBTrees.Empty();
+}
\ No newline at end of file
diff --git a/Engine/Plugins/Experimental/UVEditor/Source/UVEditorTools/Public/ToolTargets/UVEditorToolMeshInput.h b/Engine/Plugins/Experimental/UVEditor/Source/UVEditorTools/Public/ToolTargets/UVEditorToolMeshInput.h
new file mode 100644
--- /dev/null
+++ b/Engine/Plugins/Experimental/UVEditor/Source/UVEditorTools/Public/ToolTargets/UVEditorToolMeshInput.h
@@ -0,0 +1,157 @@
+// Copyright Epic Games, Inc. All Rights Reserved.
+
+#pragma once
+
+#include "CoreMinimal.h"
+
+#include "DynamicMesh/DynamicMeshAttributeSet.h" // FDynamicMeshUVOverlay
+#include "ToolTargets/ToolTarget.h"
+#include "VectorTypes.h"
+
+#include "GeometryBase.h"
+
+#include "UVEditorToolMeshInput.generated.h"
+
+class UMeshOpPreviewWithBackgroundCompute;
+class UMeshElementsVisualizer;
+PREDECLARE_GEOMETRY(class FDynamicMesh3);
+
+/**
+ * A package of the needed information for an asset being operated on by a
+ * UV editor tool. It includes a UV unwrap mesh, a mesh with the UV layer applied,
+ * and background-op-compatible previews for each. It also has convenience methods
+ * for updating all of the represenations from just one of them, using a "fast update"
+ * code path when possible.
+ * 
+ * This tool target is a bit different from usual in that it is not created
+ * by a tool target manager, and therefore doesn't have an accompanying factory.
+ * Instead, it is created by the mode, because the mode has access to the worlds
+ * in which the previews need to be created.
+ * 
+ * It's arguable whether this should even inherit from UToolTarget.
+ */
+UCLASS()
+class UVEDITORTOOLS_API UUVEditorToolMeshInput : public UToolTarget
+{
+	GENERATED_BODY()
+
+public:
+	/** 
+	 * Mesh representing the unwrapped UV layer. If the UnwrapPreview is changed via background
+	 * ops, then this mesh can be used to restart an operation as parameters change. Once a change
+	 * is completed, this mesh should be updated (i.e., it is the "canonical" unwrap mesh, though
+	 * the final UV layer truth is in the UV's of AppliedCanonical).
+	 */
+	TSharedPtr<UE::Geometry::FDynamicMesh3> UnwrapCanonical;
+
+	/**
+	 * Preview of the unwrapped UV layer, suitable for being manipulated by background ops.
+	 */
+	UPROPERTY()
+	UMeshOpPreviewWithBackgroundCompute* UnwrapPreview = nullptr;
+
+	// Note that both UnwrapCanonical and UnwrapPreview, besides having vert positions that represent
+	// a UV layer, also have a primary UV overlay that represents the same layer, to make it possible
+	// to someday texture the unwrap. This UV overlay will differ from the UV overlays in AppliedCanonical
+	// and AppliedPreview only in the parent pointers of its elements, since there will not be any elements 
+	// pointing back to the same vertex. The element id's and the triangles will be the same.
+
+	/** 
+	 * A 3d mesh with the UV layer applied. This is the canonical result that will be baked back
+	 * on the application of changes. It can also be used to reset background ops that may operate
+	 * on AppliedPreview.
+	 */
+	TSharedPtr<UE::Geometry::FDynamicMesh3> AppliedCanonical;
+
+	/**
+	 * 3d preview of the asset with the UV layer updated, suitable for use with background ops. 
+	 */
+	UPROPERTY()
+	UMeshOpPreviewWithBackgroundCompute* AppliedPreview;
+
+	/**
+	 * Optional: a wireframe to track the mesh in unwrap preview. If set, it gets updated whenever the
+	 * class updates the unwrap preview, and it is destroyed on Shutdown().
+	 * TODO: We should have a fast path for updating the wireframe...
+	 */
+	UPROPERTY()
+	UMeshElementsVisualizer* WireframeDisplay = nullptr;
+
+	// Additional needed information
+	TObjectPtr<UObject> OriginalAsset = nullptr;
+	int32 UVLayerIndex = 0;
+
+	// Mappings used for generating and baking back the unwrap.
+	TFunction<FVector3d(const FVector2f&)> UVToVertPosition;
+	TFunction<FVector2f(const FVector3d&)> VertPositionToUV;
+
+	bool InitializeMeshes(UToolTarget* Target, int32 UVLayerIndex,
+		UWorld* UnwrapWorld, UWorld* LivePreviewWorld,
+		UMaterialInterface* WorkingMaterialIn,
+		TFunction<FVector3d(const FVector2f&)> UVToVertPositionFuncIn,
+		TFunction<FVector2f(const FVector3d&)> VertPositionToUVFuncIn);
+
+	void Shutdown();
+
+	// Notes above the below convenience functions:
+	// 1. If BOTH ChangedVids/ChagnedElementIDs and ChangedTids are null, a complete update will be performed. If either
+	//   is present, then only those vids/elements/tids will be updated.
+	// 2. Generally ChangedVids/ChangedElementIDs is allowed to have new elements, since it is natural to split UVs.
+	//   However, ChangedTids must not have new Tids, because the tids form our correspondence between the unwrap mesh and
+	//   the original layer.
+	// 3. If updating the preview objects, note that the functions do not try to cancel any active computations, so an active
+	//   computation could reset things once it completes.
+
+	/**
+	 * Updates UnwrapPreview UV Overlay from UnwrapPreview vert positions. Issues a NotifyDeferredEditCompleted
+	 * for both positions and UVs.
+	 */
+	void UpdateUnwrapPreviewOverlay(const TArray<int32>* ChangedVids = nullptr, const TArray<int32>* ChangedTids = nullptr);
+
+	/**
+	 * Updates UnwrapCanonical UV Overlay from UnwrapCanonical vert positions. Issues a NotifyDeferredEditCompleted
+	 * for both positions and UVs.
+	 */
+	void UpdateUnwrapCanonicalOverlay(const TArray<int32>* ChangedVids = nullptr, const TArray<int32>* ChangedTids = nullptr);
+
+	/**
+	 * Updates the AppliedPreview from UnwrapPreview, without updating the non-preview meshes. Useful for updates during
+	 * a drag, etc, when we only care about updating the visible items.
+	 * Assumes that the overlay in UnwrapPreview is already updated.
+	 */
+	void UpdateAppliedPreviewFromUnwrapPreview(const TArray<int32>* ChangedVids = nullptr, const TArray<int32>* ChangedTids = nullptr);
+
+	/**
+	 * Updates only the UnwrapPreview from AppliedPreview, without updating the non-preview meshes. Useful for updates during
+	 * a drag, etc, when we only care about updating the visible items.
+	 */
+	void UpdateUnwrapPreviewFromAppliedPreview(
+		const TArray<int32>* ChangedElementIDs = nullptr, const TArray<int32>* ChangedTids = nullptr);
+
+	/**
+	 * Updates the non-preview meshes from their preview counterparts. Useful, for instance, after the completion of a drag
+	 * to update the canonical objects.
+	 */
+	void UpdateCanonicalFromPreviews(const TArray<int32>* ChangedVids = nullptr, const TArray<int32>* ChangedTids = nullptr);
+
+	/**
+	 * Updates the other meshes using the mesh in UnwrapPreview. Assumes that the overlay in UnwrapPreview is updated.
+	 */
+	void UpdateAllFromUnwrapPreview(const TArray<int32>* ChangedVids = nullptr, const TArray<int32>* ChangedTids = nullptr);
+
+	/**
+	 * Updates the other meshes using the mesh in UnwrapCanonical. Assumes that overlay in UnwrapCanonical is updated.
+	 */
+	void UpdateAllFromUnwrapCanonical(const TArray<int32>* ChangedVids = nullptr, const TArray<int32>* ChangedTids = nullptr);
+
+	/**
+	 * Updates the other meshes using the UV overlay in the live preview.
+	 */
+	void UpdateAllFromAppliedPreview(const TArray<int32>* ChangedElementIDs = nullptr, const TArray<int32>* ChangedTids = nullptr);
+protected:
+
+	void UpdateAppliedOverlays(const UE::Geometry::FDynamicMeshUVOverlay& SourceOverlay,
+		const TArray<int32>* ChangedVids, const TArray<int32>* ChangedTids);
+	void UpdateOtherUnwrap(const UE::Geometry::FDynamicMesh3& SourceUnwrapMesh,
+		UE::Geometry::FDynamicMesh3& DestUnwrapMesh, const TArray<int32>* ChangedVids, const TArray<int32>* ChangedTids);
+};
\ No newline at end of file
diff --git a/Engine/Plugins/Experimental/UVEditor/Source/UVEditorTools/Public/UVEditorToolUtil.h b/Engine/Plugins/Experimental/UVEditor/Source/UVEditorTools/Public/UVEditorToolUtil.h
new file mode 100644
--- /dev/null
+++ b/Engine/Plugins/Experimental/UVEditor/Source/UVEditorTools/Public/UVEditorToolUtil.h
@@ -0,0 +1,30 @@
+// Copyright Epic Games, Inc. All Rights Reserved.
+
+#pragma once
+
+#include "CoreMinimal.h"
+
+#include "DynamicMesh/DynamicMeshAttributeSet.h" //FDynamicMeshUVOverlay
+#include "VectorTypes.h"
+
+#include "GeometryBase.h"
+
+PREDECLARE_GEOMETRY(class FDynamicMesh3);
+
+namespace UE {
+namespace Geometry {
+namespace UVEditorToolUtil {
+
+	UVEDITORTOOLS_API void GenerateUVUnwrapMesh(const FDynamicMeshUVOverlay& UVOverlay, FDynamicMesh3& UnwrapMeshOut,
+		TFunctionRef<FVector3d(const FVector2f&)> UVToVertPosition);
+
+	UVEDITORTOOLS_API void UpdateUVUnwrapMesh(const FDynamicMeshUVOverlay& UVOverlayIn,
+		FDynamicMesh3& UnwrapMeshOut, TFunctionRef<FVector3d(const FVector2f&)> UVToVertPosition,
+		const TArray<int32>* ChangedElementIDs, const TArray<int32>* ChangedTids);
+
+	UVEDITORTOOLS_API void UpdateUVOverlayFromUnwrapMesh(
+		const FDynamicMesh3& UnwrapMeshIn, FDynamicMeshUVOverlay& UVOverlayOut,
+		TFunctionRef<FVector2f(const FVector3d&)> UVToVertPosition,
+		const TArray<int32>* ChangedVids = nullptr, const TArray<int32>* ChangedTids = nullptr);
+
+}}}
\ No newline at end of file
diff --git a/Engine/Plugins/Experimental/UVEditor/Source/UVEditorTools/Public/UVSelectTool.h b/Engine/Plugins/Experimental/UVEditor/Source/UVEditorTools/Public/UVSelectTool.h
--- a/Engine/Plugins/Experimental/UVEditor/Source/UVEditorTools/Public/UVSelectTool.h
+++ b/Engine/Plugins/Experimental/UVEditor/Source/UVEditorTools/Public/UVSelectTool.h
@@ -7,11 +7,14 @@
 #include "DynamicMeshAABBTree3.h"
 #include "FrameTypes.h"
 #include "InteractiveTool.h"
+//#include "InteractiveToolChange.h"
 #include "InteractiveToolBuilder.h"
 #include "TargetInterfaces/UVUnwrapDynamicMesh.h"
 
 #include "UVSelectTool.generated.h"
 
+class APreviewGeometryActor;
+class ULineSetComponent;
 class UMeshSelectionMechanic;
 class UToolTargetManager;
 class UTransformGizmo;
@@ -19,6 +22,7 @@ class UTransformProxy;
 class UUVEditorMode;
 class UUVToolStateObjectStore;
 class UPreviewMesh;
+class UUVEditorToolMeshInput;
 
 UCLASS()
 class UVEDITORTOOLS_API UUVSelectToolBuilder : public UInteractiveToolBuilder
@@ -33,12 +37,7 @@ public:
 
 	// This is a pointer so that it can be updated under the builder without
 	// having to set it in the mode after initializing targets.
-	const TArray<TObjectPtr<UPreviewMesh>>* DisplayedMeshes;
-
-	// TODO: This is one way to pass around state between tools. Another could be to
-	// hang the state onto the targets via a tool target manager, or to have something
-	// like UInteractiveToolsSelectionStoreSubsystem. Which should we do?
-	TObjectPtr<UUVToolStateObjectStore> StateObjectStore = nullptr;
+	const TArray<TObjectPtr<UUVEditorToolMeshInput>>* Targets = nullptr;
 
 protected:
 	virtual const FToolTargetTypeRequirements& GetTargetRequirements() const override;
@@ -60,8 +59,12 @@ public:
 
 	UPROPERTY(EditAnywhere, Category = Options)
 	EUVSelectToolSelectionMode SelectionMode = EUVSelectToolSelectionMode::Island;
-};
 
+	//~ TODO: Make this only visible in transform mode
+	UPROPERTY(EditAnywhere, Category = Options, AdvancedDisplay)
+	bool bUpdatePreviewDuringDrag = true;
+
+};
 
 /**
  * A tool for selecting elements of a flat FDynamicMesh corresponding to a UV layer of some asset.
@@ -77,15 +80,15 @@ class UVEDITORTOOLS_API UUVSelectTool : public UInteractiveTool
 	using FDynamicMeshAABBTree3 = UE::Geometry::FDynamicMeshAABBTree3;
 
 public:
+	virtual void SetWorld(UWorld* World) { TargetWorld = World; }
 	virtual void SetGizmoEnabled(bool bEnabledIn);
-	virtual void SetStateObjectStore(TObjectPtr<UUVToolStateObjectStore> StoreIn) { StateObjectStore = StoreIn; };
 	
 	/**
 	 * The tool will operate on the meshes given here.
 	 */
-	virtual void SetDisplayedMeshes(const TArray<TObjectPtr<UPreviewMesh>>& DisplayedMeshesIn)
+	virtual void SetTargets(const TArray<TObjectPtr<UUVEditorToolMeshInput>>& TargetsIn)
 	{
-		DisplayedMeshes = DisplayedMeshesIn;
+		Targets = TargetsIn;
 	}
 
 	// UInteractiveTool
@@ -106,12 +109,12 @@ protected:
 	virtual void GizmoTransformEnded(UTransformProxy* Proxy);
 
 	virtual void UpdateGizmo();
+	virtual void UpdateLivePreviewLines();
 
-	UPROPERTY()
-	TObjectPtr<UUVToolStateObjectStore> StateObjectStore;
+	UWorld* TargetWorld;
 
 	UPROPERTY()
-	TArray<TObjectPtr<UPreviewMesh>> DisplayedMeshes;
+	TArray<TObjectPtr<UUVEditorToolMeshInput>> Targets;
 
 	UPROPERTY()
 	TObjectPtr<UUVSelectToolProperties> Settings = nullptr;
@@ -124,12 +127,43 @@ protected:
 
 	TArray<TSharedPtr<FDynamicMeshAABBTree3>> AABBTrees;
 
-	// Placeholders: used for transforming a triangle
+	UPROPERTY()
+	APreviewGeometryActor* LivePreviewGeometryActor = nullptr;
+
+	UPROPERTY()
+	ULineSetComponent* LivePreviewLineSet = nullptr;
+
 	UE::Geometry::FFrame3d InitialGizmoFrame;
 	TArray<int32> MovingVids;
 	TArray<FVector3d> MovingVertOriginalPositions;
-	int32 SelectionMeshIndex;
+	int32 SelectionTargetIndex;
+	TArray<int32> BoundaryEids;
 
 	// We need this flag so that SetGizmoVisibility can be called before Setup() by the tool builder.
 	bool bGizmoEnabled = false;
 };
+
+///**
+// * Wraps a FDynamicMeshChange.
+// */
+//class UVEDITORTOOLS_API FUVSelectToolChange : public FToolCommandChange
+//{
+//public:
+//	FUVSelectToolChange(TUniquePtr<FDynamicMeshChange> MeshChangeIn)
+//		: MeshChange(MoveTemp(MeshChangeIn))
+//	{};
+//
+//	virtual void Apply(UObject* Object) override;
+//	virtual void Revert(UObject* Object) override;
+//	//virtual bool HasExpired(UObject* Object) const override
+//	//{
+//	//	return Cast<UGroupEdgeInsertionTool>(Object)->CurrentChangeStamp != ChangeStamp;
+//	//}
+//	virtual FString ToString() const override
+//	{
+//		return TEXT("FUVSelectToolChange");
+//	}
+//
+//protected:
+//	TUniquePtr<FDynamicMeshChange> MeshChange;
+//};
diff --git a/Engine/Plugins/Experimental/UVEditor/Source/UVEditorTools/Public/UVToolContextObjects.h b/Engine/Plugins/Experimental/UVEditor/Source/UVEditorTools/Public/UVToolContextObjects.h
new file mode 100644
--- /dev/null
+++ b/Engine/Plugins/Experimental/UVEditor/Source/UVEditorTools/Public/UVToolContextObjects.h
@@ -0,0 +1,76 @@
+// Copyright Epic Games, Inc. All Rights Reserved.
+
+#pragma once
+
+#include "CoreMinimal.h"
+
+#include "DynamicMeshAABBTree3.h"
+#include "Selection/DynamicMeshSelection.h"
+#include "GeometryBase.h"
+
+#include "UVToolContextObjects.generated.h"
+
+// TODO: Should this be spread out across multiple files?
+
+PREDECLARE_GEOMETRY(class FDynamicMesh3);
+class UInputRouter;
+class UWorld;
+
+UCLASS()
+class UVEDITORTOOLS_API UUVToolContextObject : public UObject
+{
+	GENERATED_BODY()
+public:
+};
+
+UCLASS()
+class UVEDITORTOOLS_API UUVToolLivePreviewAPI : public UUVToolContextObject
+{
+	GENERATED_BODY()
+public:
+
+	void Initialize(UWorld* WorldIn, UInputRouter* RouterIn);
+
+	UWorld* GetLivePreviewWorld() { return World.Get(); }
+	UInputRouter* GetLivePreviewInputRouter() { return InputRouter.Get(); }
+protected:
+	UPROPERTY()
+	TWeakObjectPtr<UWorld> World;
+
+	UPROPERTY()
+	TWeakObjectPtr<UInputRouter> InputRouter;
+};
+
+/** Stores a UV mesh selection */
+UCLASS()
+class UVEDITORTOOLS_API UUVToolMeshSelection : public UUVToolContextObject
+{
+	GENERATED_BODY()
+public:
+
+	TSharedPtr<UE::Geometry::FDynamicMeshSelection, ESPMode::ThreadSafe> Selection 
+		= MakeShared<UE::Geometry::FDynamicMeshSelection, ESPMode::ThreadSafe>();
+};
+
+/** Stores UV mesh AABB trees */
+UCLASS()
+class UVEDITORTOOLS_API UUVToolAABBTreeStorage : public UUVToolContextObject
+{
+	GENERATED_BODY()
+public:
+
+	void Set(UE::Geometry::FDynamicMesh3* MeshKey, TSharedPtr<UE::Geometry::FDynamicMeshAABBTree3> Tree);
+
+	TSharedPtr<UE::Geometry::FDynamicMeshAABBTree3> Get(UE::Geometry::FDynamicMesh3* MeshKey);
+
+	void Remove(UE::Geometry::FDynamicMesh3* MeshKey);
+
+	void RemoveByPredicate(TUniqueFunction<
+		bool(const TPair<UE::Geometry::FDynamicMesh3*, TSharedPtr<UE::Geometry::FDynamicMeshAABBTree3>>&)> Predicate);
+
+	void Empty();
+
+protected:
+
+	TMap<UE::Geometry::FDynamicMesh3*, TSharedPtr<UE::Geometry::FDynamicMeshAABBTree3>> AABBTrees;
+};
diff --git a/Engine/Plugins/Experimental/UVEditor/Source/UVEditorTools/Public/UVToolStateObjects.h b/Engine/Plugins/Experimental/UVEditor/Source/UVEditorTools/Public/UVToolStateObjects.h
deleted file mode 100644
diff --git a/Engine/Plugins/Experimental/UVEditor/Source/UVEditorTools/UVEditorTools.Build.cs b/Engine/Plugins/Experimental/UVEditor/Source/UVEditorTools/UVEditorTools.Build.cs
--- a/Engine/Plugins/Experimental/UVEditor/Source/UVEditorTools/UVEditorTools.Build.cs
+++ b/Engine/Plugins/Experimental/UVEditor/Source/UVEditorTools/UVEditorTools.Build.cs
@@ -45,7 +45,7 @@ public class UVEditorTools : ModuleRules
 				"Slate",
 				"SlateCore",
 					
-				"MeshConversion",
+				"MeshModelingTools",
 				"DynamicMesh",
 				"GeometryCore",
 				"ModelingComponents",
diff --git a/Engine/Plugins/Experimental/UVEditor/UVEditor.uplugin b/Engine/Plugins/Experimental/UVEditor/UVEditor.uplugin
--- a/Engine/Plugins/Experimental/UVEditor/UVEditor.uplugin
+++ b/Engine/Plugins/Experimental/UVEditor/UVEditor.uplugin
@@ -17,7 +17,8 @@
 		{
 			"Name": "UVEditor",
 			"Type": "Editor",
-			"LoadingPhase": "Default"
+			"LoadingPhase": "Default",
+			"WhitelistPlatforms": [ "Win64" ]
 		},
 		{
 			"Name": "UVEditorTools",
diff --git a/Engine/Source/Runtime/Experimental/InteractiveToolsFramework/Public/ContextObjectStore.h b/Engine/Source/Runtime/Experimental/InteractiveToolsFramework/Public/ContextObjectStore.h
--- a/Engine/Source/Runtime/Experimental/InteractiveToolsFramework/Public/ContextObjectStore.h
+++ b/Engine/Source/Runtime/Experimental/InteractiveToolsFramework/Public/ContextObjectStore.h
@@ -64,7 +64,7 @@ public:
 	template <class TObjectType>
 	bool RemoveContextObjectsOfType()
 	{
-		return RemoveToolDataObjectsOfType(TObjectType::StaticClass());
+		return RemoveContextObjectsOfType(TObjectType::StaticClass());
 	}
 
 	/**
