UVEditor: CIS fixes for non-unity builds.

#rnx

diff --git a/Engine/Plugins/Experimental/UVEditor/Source/UVEditorTools/Private/ToolTargets/UVEditorToolMeshInput.cpp b/Engine/Plugins/Experimental/UVEditor/Source/UVEditorTools/Private/ToolTargets/UVEditorToolMeshInput.cpp
--- a/Engine/Plugins/Experimental/UVEditor/Source/UVEditorTools/Private/ToolTargets/UVEditorToolMeshInput.cpp
+++ b/Engine/Plugins/Experimental/UVEditor/Source/UVEditorTools/Private/ToolTargets/UVEditorToolMeshInput.cpp
@@ -109,6 +109,16 @@ namespace UVEditorToolMeshInputLocals
 	}
 }
 
+bool UUVEditorToolMeshInput::IsValid() const
+{
+	return UnwrapCanonical
+		&& UnwrapPreview && UnwrapPreview->IsValidLowLevel()
+		&& AppliedCanonical
+		&& AppliedPreview && AppliedPreview->IsValidLowLevel()
+		&& OriginalAsset && OriginalAsset->IsValidLowLevel()
+		&& UVLayerIndex >= 0;
+}
+
 bool UUVEditorToolMeshInput::InitializeMeshes(
 	UToolTarget* Target, int32 UVLayerIndexIn,
 	UWorld* UnwrapWorld, UWorld* LivePreviewWorld,
diff --git a/Engine/Plugins/Experimental/UVEditor/Source/UVEditorTools/Private/UVToolContextObjects.cpp b/Engine/Plugins/Experimental/UVEditor/Source/UVEditorTools/Private/UVToolContextObjects.cpp
--- a/Engine/Plugins/Experimental/UVEditor/Source/UVEditorTools/Private/UVToolContextObjects.cpp
+++ b/Engine/Plugins/Experimental/UVEditor/Source/UVEditorTools/Private/UVToolContextObjects.cpp
@@ -1,7 +1,5 @@
 // Copyright Epic Games, Inc. All Rights Reserved.
 
-#pragma once
-
 #include "UVToolContextObjects.h"
 
 #include "InputRouter.h"
diff --git a/Engine/Plugins/Experimental/UVEditor/Source/UVEditorTools/Public/ToolTargets/UVEditorToolMeshInput.h b/Engine/Plugins/Experimental/UVEditor/Source/UVEditorTools/Public/ToolTargets/UVEditorToolMeshInput.h
--- a/Engine/Plugins/Experimental/UVEditor/Source/UVEditorTools/Public/ToolTargets/UVEditorToolMeshInput.h
+++ b/Engine/Plugins/Experimental/UVEditor/Source/UVEditorTools/Public/ToolTargets/UVEditorToolMeshInput.h
@@ -12,6 +12,7 @@
 
 #include "UVEditorToolMeshInput.generated.h"
 
+class UMaterialInterface;
 class UMeshOpPreviewWithBackgroundCompute;
 class UMeshElementsVisualizer;
 PREDECLARE_GEOMETRY(class FDynamicMesh3);
@@ -48,7 +49,7 @@ public:
 	 * Preview of the unwrapped UV layer, suitable for being manipulated by background ops.
 	 */
 	UPROPERTY()
-	UMeshOpPreviewWithBackgroundCompute* UnwrapPreview = nullptr;
+	TObjectPtr<UMeshOpPreviewWithBackgroundCompute> UnwrapPreview = nullptr;
 
 	// Note that both UnwrapCanonical and UnwrapPreview, besides having vert positions that represent
 	// a UV layer, also have a primary UV overlay that represents the same layer, to make it possible
@@ -67,7 +68,7 @@ public:
 	 * 3d preview of the asset with the UV layer updated, suitable for use with background ops. 
 	 */
 	UPROPERTY()
-	UMeshOpPreviewWithBackgroundCompute* AppliedPreview;
+	TObjectPtr<UMeshOpPreviewWithBackgroundCompute> AppliedPreview;
 
 	/**
 	 * Optional: a wireframe to track the mesh in unwrap preview. If set, it gets updated whenever the
@@ -75,7 +76,7 @@ public:
 	 * TODO: We should have a fast path for updating the wireframe...
 	 */
 	UPROPERTY()
-	UMeshElementsVisualizer* WireframeDisplay = nullptr;
+	TObjectPtr<UMeshElementsVisualizer> WireframeDisplay = nullptr;
 
 	// Additional needed information
 	TObjectPtr<UObject> OriginalAsset = nullptr;
@@ -148,6 +149,9 @@ public:
 	 * Updates the other meshes using the UV overlay in the live preview.
 	 */
 	void UpdateAllFromAppliedPreview(const TArray<int32>* ChangedElementIDs = nullptr, const TArray<int32>* ChangedTids = nullptr);
+
+	// UToolTarget
+	virtual bool IsValid() const override;
 protected:
 
 	void UpdateAppliedOverlays(const UE::Geometry::FDynamicMeshUVOverlay& SourceOverlay,
