UVEditor: Another file for trivial non unity fix.

#rnx

diff --git a/Engine/Plugins/Experimental/UVEditor/Source/UVEditorTools/Private/UVToolContextObjects.cpp b/Engine/Plugins/Experimental/UVEditor/Source/UVEditorTools/Private/UVToolContextObjects.cpp
--- a/Engine/Plugins/Experimental/UVEditor/Source/UVEditorTools/Private/UVToolContextObjects.cpp
+++ b/Engine/Plugins/Experimental/UVEditor/Source/UVEditorTools/Private/UVToolContextObjects.cpp
@@ -2,7 +2,8 @@
 
 #include "UVToolContextObjects.h"
 
-#include "InputRouter.h"
+#include "InputRouter.h" // Need to define this and UWorld so weak pointers know they are UThings
+#include "Engine/World.h"
 
 using namespace UE::Geometry;
 
