UVEditor: Implementing basic zoom-adaptive background grid

#rnx

diff --git a/Engine/Plugins/Experimental/MeshModelingToolset/Source/ModelingComponents/Private/Drawing/MeshDebugDrawing.cpp b/Engine/Plugins/Experimental/MeshModelingToolset/Source/ModelingComponents/Private/Drawing/MeshDebugDrawing.cpp
--- a/Engine/Plugins/Experimental/MeshModelingToolset/Source/ModelingComponents/Private/Drawing/MeshDebugDrawing.cpp
+++ b/Engine/Plugins/Experimental/MeshModelingToolset/Source/ModelingComponents/Private/Drawing/MeshDebugDrawing.cpp
@@ -5,6 +5,8 @@
 #include "FrameTypes.h"
 #include "ToolSceneQueriesUtil.h"
 
+#include "MathUtil.h"
+
 #include "SceneManagement.h" // FPrimitiveDrawInterface
 
 #include "ExplicitUseGeometryMathTypes.h"		// using UE::Geometry::(math types)
@@ -122,3 +124,64 @@ void MeshDebugDraw::DrawSimpleFixedScreenAreaGrid(
 	float GridLineSpacing = GridWidth / (float)NumGridLines;
 	DrawSimpleGrid(LocalFrame, NumGridLines, GridLineSpacing, LineWidth, Color, bDepthTested, PDI, Transform);
 }
+
+
+void MeshDebugDraw::DrawHierarchicalGrid(
+	float BaseScale, float GridZoomFactor, int32 MaxLevelDensity,
+	const FVector& WorldMaxBounds, const FVector& WorldMinBounds,
+	int32 Levels, int32 Subdivisions, TArray<FColor>& Colors,
+	const FFrame3f& LocalFrame, float LineWidth, bool bDepthTested,
+	FPrimitiveDrawInterface* PDI, const FTransform& Transform)
+{
+
+	// Determine the logrithmic scaling factor based on linear zoom factor.
+    // This allows us to track and discretely shift grid resolutions at certain zoom levels.
+	// This code assumes that we want the grid to remain stable with one logrithmic "unit"
+	// around the base zoom factor of 1.0
+	float LogZoom = FMath::LogX(Subdivisions, GridZoomFactor );
+	float LogZoomDirection = FMath::Sign(LogZoom);
+	LogZoom = FMath::Abs(LogZoom);
+	LogZoom = FMathf::Floor(LogZoom);
+	LogZoom = LogZoomDirection * LogZoom;
+
+	// Adjust grid scales based on current zoom levels
+	TArray<float> GridScales;
+	GridScales.SetNum(Levels);
+	for (int32 Level = 0; Level < Levels; ++Level)
+	{
+		GridScales[Level] = BaseScale * FMathf::Pow(Subdivisions, LogZoom - Level);
+		ensure(!FMath::IsNearlyZero(GridScales[Level]));
+	}
+	
+	//Determine the center of the drawing area for the grid, snapping to grid positions
+	FVector GridOrigin(
+		FMath::GridSnap(LocalFrame.Origin.X, GridScales[0]),
+		FMath::GridSnap(LocalFrame.Origin.Y, GridScales[0]),
+		LocalFrame.Origin.Z);
+	UE::Geometry::FFrame3f GridFrame(GridOrigin);
+
+	// Draw each level of the grid
+	for (int32 Level = 0; Level < Levels; ++Level)
+	{
+		// We automatically adjust down each level's thickness by half each time
+		float AdjustedLineWidth = LineWidth * FMathf::Pow(2, -Level);
+
+		// Compute the number of needed grid lines based on the coarsest grid's scale, this way we never "run out" of lines as we pan about
+		int32 GridLines = FMathf::Ceil(FMathf::Max(((WorldMaxBounds.X - WorldMinBounds.X + GridScales[0]*2) / GridScales[Level]),
+			                                        ((WorldMaxBounds.Y - WorldMinBounds.Y + GridScales[0]*2) / GridScales[Level])));
+
+		// If we ever have too many lines to draw, just bail. This preserves performance for large, deep grids.
+		if (GridLines > MaxLevelDensity)
+		{
+			break;
+		}
+
+		// Select our color, using the last grid color over again if we don't have enough.
+		FColor GridColor = Colors.Num() - 1 > Level ? Colors[Level] : Colors.Last();
+
+		// Finally draw one grid level.
+		MeshDebugDraw::DrawSimpleGrid(GridFrame, GridLines, GridScales[Level], AdjustedLineWidth,
+			GridColor, bDepthTested, PDI, Transform);
+	}	
+
+}
\ No newline at end of file
diff --git a/Engine/Plugins/Experimental/MeshModelingToolset/Source/ModelingComponents/Public/Drawing/MeshDebugDrawing.h b/Engine/Plugins/Experimental/MeshModelingToolset/Source/ModelingComponents/Public/Drawing/MeshDebugDrawing.h
--- a/Engine/Plugins/Experimental/MeshModelingToolset/Source/ModelingComponents/Public/Drawing/MeshDebugDrawing.h
+++ b/Engine/Plugins/Experimental/MeshModelingToolset/Source/ModelingComponents/Public/Drawing/MeshDebugDrawing.h
@@ -103,4 +103,27 @@ namespace MeshDebugDraw
 		float LineWidth, FColor Color, bool bDepthTested,
 		FPrimitiveDrawInterface* PDI, const FTransform& Transform);
 
+
+	/**
+	 * Draw a zoomable, hierarchical grid with configurable depth and density terms
+	 * @param BaseGridScale At 1.0 zoom, what is the distance between the top most (level 0) grid lines 
+	 * @param GridZoomFactor A pseudo zoom factor, where a value of 1.0 represents the baseline grid scaling with BaseScale the distance between coarsest grid lines.
+	 * @param MaxLevelDensity The maximum number of lines to draw at any given grid level
+	 * @param Levels Number of heirarchy levels to draw. 1 level is equivalent to a simple grid
+	 * @param Subdivisions Number of gridlines inserted between lines of a higher level
+	 * @param Colors Array containing colors to use for the grid levels. Must include at least one color. If fewer colors than levels are provided, last color will be repeated as needed.
+	 * @param LocalFrame Pre-transform frame of grid (grid lies in XY plane).
+	 * @param GridLines number of grid lines. If odd, there is a center-line, if even then frame center is at center of a grid square
+	 * @param GridLineSpacing spacing size between grid lines
+	 * @param LineWidth thickness of the lines in screen space
+	 * @param DepthPriority drawing depth priority
+	 * @param PDI drawing interface
+	 * @param Transform transform applied to LocalFrame. Pass as Identity() if you have world frame.		
+	*/
+	void MODELINGCOMPONENTS_API DrawHierarchicalGrid(
+		float BaseGridScale, float GridZoomFactor, int32 MaxLevelDensity,
+		const FVector& WorldMaxBounds, const FVector& WorldMinBounds,
+		int32 Levels, int32 Subdivisions, TArray<FColor>& Colors,
+		const FFrame3f& LocalFrame, float LineWidth, bool bDepthTested,
+		FPrimitiveDrawInterface* PDI, const FTransform& Transform);
 }
\ No newline at end of file
diff --git a/Engine/Plugins/Experimental/UVEditor/Source/UVEditor/Private/UVEditor2DViewportClient.cpp b/Engine/Plugins/Experimental/UVEditor/Source/UVEditor/Private/UVEditor2DViewportClient.cpp
--- a/Engine/Plugins/Experimental/UVEditor/Source/UVEditor/Private/UVEditor2DViewportClient.cpp
+++ b/Engine/Plugins/Experimental/UVEditor/Source/UVEditor/Private/UVEditor2DViewportClient.cpp
@@ -2,11 +2,14 @@
 
 #include "UVEditor2DViewportClient.h"
 
+#include "UVEditorMode.h"
 #include "BaseBehaviors/ClickDragBehavior.h"
 #include "BaseBehaviors/MouseWheelBehavior.h"
 //#include "Drawing/DensityAdjustingGrid.h"
 #include "EditorModeManager.h"
 #include "EdModeInteractiveToolsContext.h"
+#include "Drawing/MeshDebugDrawing.h"
+#include "FrameTypes.h"
 
 FUVEditor2DViewportClient::FUVEditor2DViewportClient(FEditorModeTools* InModeTools,
 	FPreviewScene* InPreviewScene, const TWeakPtr<SEditorViewport>& InEditorViewportWidget)
@@ -72,16 +75,52 @@ bool FUVEditor2DViewportClient::InputKey(FViewport* InViewport, int32 Controller
 // We don't override that top-level function so that it can do whatever view calculations it needs to do.
 void FUVEditor2DViewportClient::Draw(const FSceneView* View, FPrimitiveDrawInterface* PDI)
 {
-	// Draw the axes
-	const float AxisThickness = 1.0;
-	PDI->DrawLine(FVector(-HALF_WORLD_MAX, 0, 0), FVector(HALF_WORLD_MAX, 0, 0), FLinearColor::Red, SDPG_World, AxisThickness, 0, true);
-	PDI->DrawLine(FVector(0, -HALF_WORLD_MAX, 0), FVector(0, HALF_WORLD_MAX, 0), FLinearColor::Green, SDPG_World, AxisThickness, 0, true);
-
-	// TODO: Eventually we want a proper grid. For now, draw the unit boundary
-	const float UVScale = 1000;
-	PDI->DrawLine(FVector(UVScale, 0, 0), FVector(UVScale, UVScale, 0), FLinearColor::Gray, SDPG_World, AxisThickness, 0, true);
-	PDI->DrawLine(FVector(0, UVScale, 0), FVector(UVScale, UVScale, 0), FLinearColor::Gray, SDPG_World, AxisThickness, 0, true);
-	//Grid->Render(View, PDI);
+	// Basic scaling amount
+	const float UVScale = UUVEditorMode::GetUVMeshScalingFactor();
+	
+	// Line thickness parameters
+	const float AxisThickness = 2;
+	const float GridMajorThickness = 1.0;
+
+    // Line color scheme parameters
+	const FLinearColor XAxisColor = FLinearColor::Red;
+	const FLinearColor YAxisColor = FLinearColor::Green;
+	const FLinearColor GridMajorColor = FLinearColor::Gray;
+	const FLinearColor GridMinorColor = FLinearColor::Gray;
+
+	// Determine important geometry of the viewport for creating grid lines
+	FVector WorldCenterPoint( 0,0,0 );
+	FVector4 WorldToScreenCenter = View->WorldToScreen(WorldCenterPoint);	
+	float ZoomFactor = WorldToScreenCenter.W;
+	FVector4 MaxScreen(1 * ZoomFactor, 1 * ZoomFactor, 0, 1);
+	FVector4 MinScreen(-1 * ZoomFactor, -1 * ZoomFactor, 0, 1);
+	FVector WorldBoundsMax = View->ScreenToWorld(MaxScreen);
+	FVector WorldBoundsMin = View->ScreenToWorld(MinScreen);
+	FVector ViewLoc = GetViewLocation();
+	ViewLoc.Z = 0.0; // We are treating the scene like a 2D plane, so we'll clamp the Z position here to 
+	               // 0 as a simple projection step just in case.
+
+	// Setup and call grid calling function
+	UE::Geometry::FFrame3f LocalFrame(ViewLoc);
+	FTransform Transform;
+	TArray<FColor> Colors;
+	Colors.Push(GridMajorColor.ToRGBE());
+	Colors.Push(GridMinorColor.ToRGBE());
+	MeshDebugDraw::DrawHierarchicalGrid(UVScale, ZoomFactor/UVScale,
+		500, // Maximum density of lines to draw per level before skipping the level
+		WorldBoundsMax, WorldBoundsMin,
+		3, // Number of levels to draw
+		4, // Number of subdivisions per level
+		Colors,
+		LocalFrame, GridMajorThickness, true,
+		PDI, Transform);
+
+
+	float AxisExtent = FMathf::Max(UVScale, FMathf::Min(WorldBoundsMax.Y, WorldBoundsMax.X));
+
+	// Draw colored axis lines
+	PDI->DrawLine(FVector(0, 0, 0), FVector(AxisExtent, 0, 0), FLinearColor::Red, SDPG_World, AxisThickness, 0, true);
+	PDI->DrawLine(FVector(0, 0, 0), FVector(0, AxisExtent, 0), FLinearColor::Green, SDPG_World, AxisThickness, 0, true);
 
 	// TODO: Draw a little UV axis thing in the lower left, like the XYZ things that normal viewports have.
 
