UVEditor: Avoid graphical issues with grid drawing at extreme zoom levels by disabling grid

#rnx

diff --git a/Engine/Plugins/Experimental/UVEditor/Source/UVEditor/Private/UVEditor2DViewportClient.cpp b/Engine/Plugins/Experimental/UVEditor/Source/UVEditor/Private/UVEditor2DViewportClient.cpp
--- a/Engine/Plugins/Experimental/UVEditor/Source/UVEditor/Private/UVEditor2DViewportClient.cpp
+++ b/Engine/Plugins/Experimental/UVEditor/Source/UVEditor/Private/UVEditor2DViewportClient.cpp
@@ -100,21 +100,24 @@ void FUVEditor2DViewportClient::Draw(const FSceneView* View, FPrimitiveDrawInter
 	ViewLoc.Z = 0.0; // We are treating the scene like a 2D plane, so we'll clamp the Z position here to 
 	               // 0 as a simple projection step just in case.
 
-	// Setup and call grid calling function
-	UE::Geometry::FFrame3f LocalFrame(ViewLoc);
-	FTransform Transform;
-	TArray<FColor> Colors;
-	Colors.Push(GridMajorColor.ToRGBE());
-	Colors.Push(GridMinorColor.ToRGBE());
-	MeshDebugDraw::DrawHierarchicalGrid(UVScale, ZoomFactor/UVScale,
-		500, // Maximum density of lines to draw per level before skipping the level
-		WorldBoundsMax, WorldBoundsMin,
-		3, // Number of levels to draw
-		4, // Number of subdivisions per level
-		Colors,
-		LocalFrame, GridMajorThickness, true,
-		PDI, Transform);
-
+	// Prevent grid from drawing if we are too close or too far, in order to avoid potential graphical issues.
+	if (ZoomFactor < 100000 && ZoomFactor > 1)
+	{
+		// Setup and call grid calling function
+		UE::Geometry::FFrame3f LocalFrame(ViewLoc);
+		FTransform Transform;
+		TArray<FColor> Colors;
+		Colors.Push(GridMajorColor.ToRGBE());
+		Colors.Push(GridMinorColor.ToRGBE());
+		MeshDebugDraw::DrawHierarchicalGrid(UVScale, ZoomFactor / UVScale,
+			500, // Maximum density of lines to draw per level before skipping the level
+			WorldBoundsMax, WorldBoundsMin,
+			3, // Number of levels to draw
+			4, // Number of subdivisions per level
+			Colors,
+			LocalFrame, GridMajorThickness, true,
+			PDI, Transform);
+	}
 
 	float AxisExtent = FMathf::Max(UVScale, FMathf::Min(WorldBoundsMax.Y, WorldBoundsMax.X));
 
