UVEditor: Improving performance of transforms for nanite scale meshes

#rnx

diff --git a/Engine/Plugins/Experimental/UVEditor/Source/UVEditorTools/Private/ToolTargets/UVEditorToolMeshInput.cpp b/Engine/Plugins/Experimental/UVEditor/Source/UVEditorTools/Private/ToolTargets/UVEditorToolMeshInput.cpp
--- a/Engine/Plugins/Experimental/UVEditor/Source/UVEditorTools/Private/ToolTargets/UVEditorToolMeshInput.cpp
+++ b/Engine/Plugins/Experimental/UVEditor/Source/UVEditorTools/Private/ToolTargets/UVEditorToolMeshInput.cpp
@@ -17,7 +17,7 @@ using namespace UE::Geometry;
 namespace UVEditorToolMeshInputLocals
 {
 	void CopyMeshPositions(const UE::Geometry::FDynamicMesh3& MeshIn, UE::Geometry::FDynamicMesh3& MeshOut,
-		const TArray<int32>* ChangedVids, const TArray<int32>* ChangedTids)
+		const TArray<int32>* ChangedVids, const TArray<int32>* ChangedConnectivityTids)
 	{
 		auto UpdateOrInsertVid = [&MeshIn, &MeshOut](int32 Vid) {
 			if (MeshOut.IsVertex(Vid))
@@ -30,7 +30,7 @@ namespace UVEditorToolMeshInputLocals
 			}
 		};
 
-		if (!ChangedVids && !ChangedTids)
+		if (!ChangedVids && !ChangedConnectivityTids)
 		{
 			for (int32 Vid : MeshIn.VertexIndicesItr())
 			{
@@ -52,20 +52,20 @@ namespace UVEditorToolMeshInputLocals
 			}
 		}
 
-		if (ChangedTids)
+		if (ChangedConnectivityTids)
 		{
-			for (int32 Tid : *ChangedTids)
+			for (int32 Tid : *ChangedConnectivityTids                                         )
 			{
 				MeshOut.SetTriangle(Tid, MeshIn.GetTriangle(Tid));
 			}
 		}
 	}
 
-	// Aside from potentially  doing partial updates based on ChangedElements and ChangedTids, this
+	// Aside from potentially  doing partial updates based on ChangedElements and ChangedConnectivityTids, this
 	// function also differs from a simple overlay copy in that it does not copy the parent pointer arrays.
 	// Also it assumes the existence of all the Tids.
 	void CopyMeshOverlay(const UE::Geometry::FDynamicMeshUVOverlay& OverlayIn, UE::Geometry::FDynamicMeshUVOverlay& OverlayOut,
-		const TArray<int32>* ChangedElements, const TArray<int32>* ChangedTids)
+		const TArray<int32>* ChangedElements, const TArray<int32>* ChangedConnectivityTids)
 	{
 		auto UpdateOrInsertElement = [&OverlayIn, &OverlayOut](int32 ElementID) {
 			if (OverlayOut.IsElement(ElementID))
@@ -78,7 +78,7 @@ namespace UVEditorToolMeshInputLocals
 			}
 		};
 
-		if (!ChangedElements && !ChangedTids)
+		if (!ChangedElements && !ChangedConnectivityTids)
 		{
 			for (int32 ElementID : OverlayIn.ElementIndicesItr())
 			{
@@ -99,9 +99,9 @@ namespace UVEditorToolMeshInputLocals
 				UpdateOrInsertElement(ElementID);
 			}
 		}
-		if (ChangedTids)
+		if (ChangedConnectivityTids)
 		{
-			for (int32 Tid : *ChangedTids)
+			for (int32 Tid : *ChangedConnectivityTids)
 			{
 				OverlayOut.SetTriangle(Tid, OverlayIn.GetTriangle(Tid));
 			}
@@ -192,17 +192,24 @@ void UUVEditorToolMeshInput::Shutdown()
 	OriginalAsset = nullptr;
 }
 
-void UUVEditorToolMeshInput::UpdateUnwrapPreviewOverlay(const TArray<int32>* ChangedVids, const TArray<int32>* ChangedTids)
+void UUVEditorToolMeshInput::UpdateUnwrapPreviewOverlay(const TArray<int32>* ChangedVids, const TArray<int32>* ChangedConnectivityTids, const TArray<int32>* FastRenderUpdateTids)
 {
-	UnwrapPreview->PreviewMesh->DeferredEditMesh([this, ChangedVids, ChangedTids](FDynamicMesh3& Mesh)
+	UnwrapPreview->PreviewMesh->DeferredEditMesh([this, ChangedVids, ChangedConnectivityTids](FDynamicMesh3& Mesh)
 	{
 		FDynamicMeshUVOverlay* DestOverlay = Mesh.Attributes()->PrimaryUV();
 		UVEditorToolUtil::UpdateUVOverlayFromUnwrapMesh(Mesh, *DestOverlay, VertPositionToUV,
-			ChangedVids, ChangedTids);
+			ChangedVids, ChangedConnectivityTids);
 	}, false);
 
-	UnwrapPreview->PreviewMesh->NotifyDeferredEditCompleted(UPreviewMesh::ERenderUpdateMode::FastUpdate,
-		EMeshRenderAttributeFlags::Positions | EMeshRenderAttributeFlags::VertexUVs, true); // Assuming user changed positions and didn't notify yet
+	if (FastRenderUpdateTids) {
+		UnwrapPreview->PreviewMesh->NotifyRegionDeferredEditCompleted(*FastRenderUpdateTids,
+			EMeshRenderAttributeFlags::Positions | EMeshRenderAttributeFlags::VertexUVs); // Assuming user changed positions and didn't notify yet
+	} 
+	else
+	{
+		UnwrapPreview->PreviewMesh->NotifyDeferredEditCompleted(UPreviewMesh::ERenderUpdateMode::FastUpdate,
+			EMeshRenderAttributeFlags::Positions | EMeshRenderAttributeFlags::VertexUVs, true); // Assuming user changed positions and didn't notify yet
+	}
 
 	if (WireframeDisplay)
 	{
@@ -210,14 +217,14 @@ void UUVEditorToolMeshInput::UpdateUnwrapPreviewOverlay(const TArray<int32>* Cha
 	}
 }
 
-void UUVEditorToolMeshInput::UpdateUnwrapCanonicalOverlay(const TArray<int32>* ChangedVids, const TArray<int32>* ChangedTids)
+void UUVEditorToolMeshInput::UpdateUnwrapCanonicalOverlay(const TArray<int32>* ChangedVids, const TArray<int32>* ChangedConnectivityTids)
 {
 	FDynamicMeshUVOverlay* DestOverlay = UnwrapCanonical->Attributes()->PrimaryUV();
 	UVEditorToolUtil::UpdateUVOverlayFromUnwrapMesh(*UnwrapCanonical, *DestOverlay, VertPositionToUV,
-		ChangedVids, ChangedTids);
+		ChangedVids, ChangedConnectivityTids);
 }
 
-void UUVEditorToolMeshInput::UpdateAppliedPreviewFromUnwrapPreview(const TArray<int32>* ChangedVids, const TArray<int32>* ChangedTids)
+void UUVEditorToolMeshInput::UpdateAppliedPreviewFromUnwrapPreview(const TArray<int32>* ChangedVids, const TArray<int32>* ChangedConnectivityTids, const TArray<int32>* FastRenderUpdateTids)
 {
 	using namespace UVEditorToolMeshInputLocals;
 
@@ -225,33 +232,49 @@ void UUVEditorToolMeshInput::UpdateAppliedPreviewFromUnwrapPreview(const TArray<
 
 	// We need to update the overlay in AppliedPreview. Assuming that the overlay in UnwrapPreview is updated, we can
 	// just copy that overlay (using our function that doesn't copy element parent pointers)
-	AppliedPreview->PreviewMesh->DeferredEditMesh([this, SourceUnwrapMesh, ChangedVids, ChangedTids](FDynamicMesh3& Mesh)
+	AppliedPreview->PreviewMesh->DeferredEditMesh([this, SourceUnwrapMesh, ChangedVids, ChangedConnectivityTids](FDynamicMesh3& Mesh)
 	{
 			const FDynamicMeshUVOverlay* SourceOverlay = SourceUnwrapMesh->Attributes()->PrimaryUV();
 			FDynamicMeshUVOverlay* DestOverlay = Mesh.Attributes()->PrimaryUV();
 
-			CopyMeshOverlay(*SourceOverlay, *DestOverlay, ChangedVids, ChangedTids);
+			CopyMeshOverlay(*SourceOverlay, *DestOverlay, ChangedVids, ChangedConnectivityTids);
 	}, false);
-	AppliedPreview->PreviewMesh->NotifyDeferredEditCompleted(UPreviewMesh::ERenderUpdateMode::FastUpdate,
-		EMeshRenderAttributeFlags::VertexUVs, false);
+
+	if (FastRenderUpdateTids) {
+		AppliedPreview->PreviewMesh->NotifyRegionDeferredEditCompleted(*FastRenderUpdateTids,
+			EMeshRenderAttributeFlags::VertexUVs); // Assuming user changed positions and didn't notify yet
+	}
+	else
+	{
+		AppliedPreview->PreviewMesh->NotifyDeferredEditCompleted(UPreviewMesh::ERenderUpdateMode::FastUpdate,
+			EMeshRenderAttributeFlags::VertexUVs, false); // Assuming user changed positions and didn't notify yet
+	}
 }
 
-void UUVEditorToolMeshInput::UpdateUnwrapPreviewFromAppliedPreview(const TArray<int32>* ChangedElementIDs, const TArray<int32>* ChangedTids)
+void UUVEditorToolMeshInput::UpdateUnwrapPreviewFromAppliedPreview(const TArray<int32>* ChangedElementIDs, const TArray<int32>* ChangedConnectivityTids, const TArray<int32>* FastRenderUpdateTids)
 {
 	using namespace UVEditorToolMeshInputLocals;
 
 	// Convert the AppliedPreview UV overlay to positions in UnwrapPreview
 	const FDynamicMeshUVOverlay* SourceOverlay = AppliedPreview->PreviewMesh->GetMesh()->Attributes()->PrimaryUV();
-	UnwrapPreview->PreviewMesh->DeferredEditMesh([this, SourceOverlay, ChangedElementIDs, ChangedTids](FDynamicMesh3& Mesh)
+	UnwrapPreview->PreviewMesh->DeferredEditMesh([this, SourceOverlay, ChangedElementIDs, ChangedConnectivityTids](FDynamicMesh3& Mesh)
 	{
-		UVEditorToolUtil::UpdateUVUnwrapMesh(*SourceOverlay, Mesh, UVToVertPosition, ChangedElementIDs, ChangedTids);
+		UVEditorToolUtil::UpdateUVUnwrapMesh(*SourceOverlay, Mesh, UVToVertPosition, ChangedElementIDs, ChangedConnectivityTids);
 
 		// Also copy the actual overlay
 		FDynamicMeshUVOverlay* DestOverlay = Mesh.Attributes()->PrimaryUV();
-		CopyMeshOverlay(*SourceOverlay, *DestOverlay, ChangedElementIDs, ChangedTids);
+		CopyMeshOverlay(*SourceOverlay, *DestOverlay, ChangedElementIDs, ChangedConnectivityTids);
 	}, false);
-	UnwrapPreview->PreviewMesh->NotifyDeferredEditCompleted(UPreviewMesh::ERenderUpdateMode::FastUpdate,
-		EMeshRenderAttributeFlags::Positions | EMeshRenderAttributeFlags::VertexUVs, true);
+
+	if (FastRenderUpdateTids) {
+		UnwrapPreview->PreviewMesh->NotifyRegionDeferredEditCompleted(*FastRenderUpdateTids,
+			EMeshRenderAttributeFlags::Positions | EMeshRenderAttributeFlags::VertexUVs); // Assuming user changed positions and didn't notify yet
+	}
+	else
+	{
+		UnwrapPreview->PreviewMesh->NotifyDeferredEditCompleted(UPreviewMesh::ERenderUpdateMode::FastUpdate,
+			EMeshRenderAttributeFlags::Positions | EMeshRenderAttributeFlags::VertexUVs, true); // Assuming user changed positions and didn't notify yet
+	}
 
 	if (WireframeDisplay)
 	{
@@ -259,47 +282,54 @@ void UUVEditorToolMeshInput::UpdateUnwrapPreviewFromAppliedPreview(const TArray<
 	}
 }
 
-void UUVEditorToolMeshInput::UpdateCanonicalFromPreviews(const TArray<int32>* ChangedVids, const TArray<int32>* ChangedTids)
+void UUVEditorToolMeshInput::UpdateCanonicalFromPreviews(const TArray<int32>* ChangedVids, const TArray<int32>* ChangedConnectivityTids)
 {
 	using namespace UVEditorToolMeshInputLocals;
 
 	// Update UnwrapCanonical from UnwrapPreview
-	UpdateOtherUnwrap(*UnwrapPreview->PreviewMesh->GetMesh(), *UnwrapCanonical, ChangedVids, ChangedTids);
+	UpdateOtherUnwrap(*UnwrapPreview->PreviewMesh->GetMesh(), *UnwrapCanonical, ChangedVids, ChangedConnectivityTids);
 	
 	// Update the overlay in AppliedCanonical from overlay in AppliedPreview
 	const FDynamicMeshUVOverlay* SourceOverlay = AppliedPreview->PreviewMesh->GetMesh()->Attributes()->GetUVLayer(UVLayerIndex);
 	FDynamicMeshUVOverlay* DestOverlay = AppliedCanonical->Attributes()->GetUVLayer(UVLayerIndex);
-	CopyMeshOverlay(*SourceOverlay, *DestOverlay, ChangedVids, ChangedTids);
+	CopyMeshOverlay(*SourceOverlay, *DestOverlay, ChangedVids, ChangedConnectivityTids);
 }
 
 void UUVEditorToolMeshInput::UpdateAllFromUnwrapPreview(
-	const TArray<int32>* ChangedVids, const TArray<int32>* ChangedTids)
+	const TArray<int32>* ChangedVids, const TArray<int32>* ChangedConnectivityTids, const TArray<int32>* FastRenderUpdateTids)
 {
 	const FDynamicMesh3* ReferenceMesh = UnwrapPreview->PreviewMesh->GetMesh();
 
 	// Update UnwrapCanonical
-	UpdateOtherUnwrap(*ReferenceMesh, *UnwrapCanonical, ChangedVids, ChangedTids);
+	UpdateOtherUnwrap(*ReferenceMesh, *UnwrapCanonical, ChangedVids, ChangedConnectivityTids);
 
 	// Update the applied meshes
 	const FDynamicMeshUVOverlay* SourceOverlay = ReferenceMesh->Attributes()->PrimaryUV();
-	UpdateAppliedOverlays(*SourceOverlay, ChangedVids, ChangedTids);
+	UpdateAppliedOverlays(*SourceOverlay, ChangedVids, ChangedConnectivityTids, FastRenderUpdateTids);
 }
 
 void UUVEditorToolMeshInput::UpdateAllFromUnwrapCanonical(
-	const TArray<int32>* ChangedVids, const TArray<int32>* ChangedTids)
+	const TArray<int32>* ChangedVids, const TArray<int32>* ChangedConnectivityTids, const TArray<int32>* FastRenderUpdateTids)
 {
 	// Update UnwrapPreview
-	UnwrapPreview->PreviewMesh->DeferredEditMesh([this, ChangedVids, ChangedTids](FDynamicMesh3& Mesh)
+	UnwrapPreview->PreviewMesh->DeferredEditMesh([this, ChangedVids, ChangedConnectivityTids](FDynamicMesh3& Mesh)
 	{
-		UpdateOtherUnwrap(*UnwrapCanonical, Mesh, ChangedVids, ChangedTids);
+		UpdateOtherUnwrap(*UnwrapCanonical, Mesh, ChangedVids, ChangedConnectivityTids);
 	}, false);
-	UnwrapPreview->PreviewMesh->NotifyDeferredEditCompleted(UPreviewMesh::ERenderUpdateMode::FastUpdate,
-		EMeshRenderAttributeFlags::Positions | EMeshRenderAttributeFlags::VertexUVs, true);
-	
+
+	if (FastRenderUpdateTids) {
+		UnwrapPreview->PreviewMesh->NotifyRegionDeferredEditCompleted(*FastRenderUpdateTids,
+			EMeshRenderAttributeFlags::Positions | EMeshRenderAttributeFlags::VertexUVs); // Assuming user changed positions and didn't notify yet
+	}
+	else
+	{
+		UnwrapPreview->PreviewMesh->NotifyDeferredEditCompleted(UPreviewMesh::ERenderUpdateMode::FastUpdate,
+			EMeshRenderAttributeFlags::Positions | EMeshRenderAttributeFlags::VertexUVs, true); // Assuming user changed positions and didn't notify yet
+	}
 
 	// Update the applied meshes
 	FDynamicMeshUVOverlay* SourceOverlay = UnwrapCanonical->Attributes()->PrimaryUV();
-	UpdateAppliedOverlays(*SourceOverlay, ChangedVids, ChangedTids);
+	UpdateAppliedOverlays(*SourceOverlay, ChangedVids, ChangedConnectivityTids, FastRenderUpdateTids);
 
 	if (WireframeDisplay)
 	{
@@ -308,7 +338,7 @@ void UUVEditorToolMeshInput::UpdateAllFromUnwrapCanonical(
 }
 
 void UUVEditorToolMeshInput::UpdateAllFromAppliedPreview(
-	const TArray<int32>* ChangedElementIDs, const TArray<int32>* ChangedTids)
+	const TArray<int32>* ChangedElementIDs, const TArray<int32>* ChangedConnectivityTids, const TArray<int32>* FastRenderUpdateTids)
 {
 	using namespace UVEditorToolMeshInputLocals;
 
@@ -316,18 +346,26 @@ void UUVEditorToolMeshInput::UpdateAllFromAppliedPreview(
 	const FDynamicMeshUVOverlay* SourceOverlay = AppliedPreview->PreviewMesh->GetMesh()
 		->Attributes()->GetUVLayer(UVLayerIndex);
 	FDynamicMeshUVOverlay* DestOverlay = AppliedCanonical->Attributes()->GetUVLayer(UVLayerIndex);
-	CopyMeshOverlay(*SourceOverlay, *DestOverlay, ChangedElementIDs, ChangedTids);
+	CopyMeshOverlay(*SourceOverlay, *DestOverlay, ChangedElementIDs, ChangedConnectivityTids);
 
 	// Update UnwrapCanonical
 	UVEditorToolUtil::UpdateUVUnwrapMesh(*SourceOverlay, *UnwrapCanonical, 
-		UVToVertPosition, ChangedElementIDs, ChangedTids);
+		UVToVertPosition, ChangedElementIDs, ChangedConnectivityTids);
 
 	// Update UnwrapPreview from UnwrapCanonical
-	UnwrapPreview->PreviewMesh->DeferredEditMesh([this, ChangedElementIDs, ChangedTids](FDynamicMesh3& Mesh) {
-		UpdateOtherUnwrap(*UnwrapCanonical, Mesh, ChangedElementIDs, ChangedTids);
+	UnwrapPreview->PreviewMesh->DeferredEditMesh([this, ChangedElementIDs, ChangedConnectivityTids](FDynamicMesh3& Mesh) {
+		UpdateOtherUnwrap(*UnwrapCanonical, Mesh, ChangedElementIDs, ChangedConnectivityTids);
 	}, false);
-	UnwrapPreview->PreviewMesh->NotifyDeferredEditCompleted(UPreviewMesh::ERenderUpdateMode::FastUpdate,
-		EMeshRenderAttributeFlags::Positions | EMeshRenderAttributeFlags::VertexUVs, true);
+
+	if (FastRenderUpdateTids) {
+		UnwrapPreview->PreviewMesh->NotifyRegionDeferredEditCompleted(*FastRenderUpdateTids,
+			EMeshRenderAttributeFlags::Positions | EMeshRenderAttributeFlags::VertexUVs); // Assuming user changed positions and didn't notify yet
+	}
+	else
+	{
+		UnwrapPreview->PreviewMesh->NotifyDeferredEditCompleted(UPreviewMesh::ERenderUpdateMode::FastUpdate,
+			EMeshRenderAttributeFlags::Positions | EMeshRenderAttributeFlags::VertexUVs, true); // Assuming user changed positions and didn't notify yet
+	}
 
 	if (WireframeDisplay)
 	{
@@ -339,41 +377,49 @@ void UUVEditorToolMeshInput::UpdateAllFromAppliedPreview(
  * Helper function. Updates the UV overlays of AppliedCanonical and AppliedPreview.
  */
 void UUVEditorToolMeshInput::UpdateAppliedOverlays(const FDynamicMeshUVOverlay& SourceOverlay,
-	const TArray<int32>* ChangedVids, const TArray<int32>* ChangedTids)
+	const TArray<int32>* ChangedVids, const TArray<int32>* ChangedConnectivityTids, const TArray<int32>* FastRenderUpdateTids)
 {
 	using namespace UVEditorToolMeshInputLocals;
 
 	// Update UnwrappedMeshCanonical
 	FDynamicMeshUVOverlay* DestOverlay = AppliedCanonical->Attributes()->GetUVLayer(UVLayerIndex);
-	CopyMeshOverlay(SourceOverlay, *DestOverlay, ChangedVids, ChangedTids);
+	CopyMeshOverlay(SourceOverlay, *DestOverlay, ChangedVids, ChangedConnectivityTids);
 
 	// Update AppliedPreview
-	AppliedPreview->PreviewMesh->DeferredEditMesh([this, &SourceOverlay, ChangedVids, ChangedTids](FDynamicMesh3& Mesh)
+	AppliedPreview->PreviewMesh->DeferredEditMesh([this, &SourceOverlay, ChangedVids, ChangedConnectivityTids](FDynamicMesh3& Mesh)
 	{
 		FDynamicMeshUVOverlay* DestOverlay = Mesh.Attributes()->GetUVLayer(UVLayerIndex);
-		CopyMeshOverlay(SourceOverlay, *DestOverlay, ChangedVids, ChangedTids);
+		CopyMeshOverlay(SourceOverlay, *DestOverlay, ChangedVids, ChangedConnectivityTids);
 	}, false);
-	AppliedPreview->PreviewMesh->NotifyDeferredEditCompleted(UPreviewMesh::ERenderUpdateMode::FastUpdate,
-		EMeshRenderAttributeFlags::VertexUVs, false);
+
+	if (FastRenderUpdateTids) {
+		AppliedPreview->PreviewMesh->NotifyRegionDeferredEditCompleted(*FastRenderUpdateTids,
+			EMeshRenderAttributeFlags::Positions | EMeshRenderAttributeFlags::VertexUVs); // Assuming user changed positions and didn't notify yet
+	}
+	else
+	{
+		AppliedPreview->PreviewMesh->NotifyDeferredEditCompleted(UPreviewMesh::ERenderUpdateMode::FastUpdate,
+			EMeshRenderAttributeFlags::VertexUVs, false); // Assuming user changed positions and didn't notify yet
+	}
 }
 
 /**
  * Helper function. Uses the positions and UV overlay of one unwrap mesh to update another one.
  */
 void UUVEditorToolMeshInput::UpdateOtherUnwrap(const FDynamicMesh3& SourceUnwrapMesh,
-	FDynamicMesh3& DestUnwrapMesh, const TArray<int32>* ChangedVids, const TArray<int32>* ChangedTids)
+	FDynamicMesh3& DestUnwrapMesh, const TArray<int32>* ChangedVids, const TArray<int32>* ChangedConnectivityTids)
 {
 	using namespace UVEditorToolMeshInputLocals;
 
-	if (!ChangedVids && !ChangedTids)
+	if (!ChangedVids && !ChangedConnectivityTids)
 	{
 		DestUnwrapMesh.Copy(SourceUnwrapMesh, false, false, false, true); // Copy positions and UVs
 		return;
 	}
 
-	CopyMeshPositions(SourceUnwrapMesh, DestUnwrapMesh, ChangedVids, ChangedTids);
+	CopyMeshPositions(SourceUnwrapMesh, DestUnwrapMesh, ChangedVids, ChangedConnectivityTids);
 
 	const FDynamicMeshUVOverlay* SourceOverlay = SourceUnwrapMesh.Attributes()->PrimaryUV();
 	FDynamicMeshUVOverlay* DestOverlay = DestUnwrapMesh.Attributes()->PrimaryUV();
-	CopyMeshOverlay(*SourceOverlay, *DestOverlay, ChangedVids, ChangedTids);
+	CopyMeshOverlay(*SourceOverlay, *DestOverlay, ChangedVids, ChangedConnectivityTids);
 }
diff --git a/Engine/Plugins/Experimental/UVEditor/Source/UVEditorTools/Private/UVSelectTool.cpp b/Engine/Plugins/Experimental/UVEditor/Source/UVEditorTools/Private/UVSelectTool.cpp
--- a/Engine/Plugins/Experimental/UVEditor/Source/UVEditorTools/Private/UVSelectTool.cpp
+++ b/Engine/Plugins/Experimental/UVEditor/Source/UVEditorTools/Private/UVSelectTool.cpp
@@ -200,6 +200,7 @@ void UUVSelectTool::OnSelectionChanged()
 
 	SelectionTargetIndex = -1;
 	MovingVids.Reset();
+	SelectedTids.Reset();
 	BoundaryEids.Reset();
 	if (!Selection.IsEmpty())
 	{
@@ -216,6 +217,7 @@ void UUVSelectTool::OnSelectionChanged()
 
 		// Note the selected vids
 		TSet<int32> VidSet;
+		TSet<int32> TidSet;
 		if (Selection.Type == FDynamicMeshSelection::EType::Triangle)
 		{
 			const FDynamicMesh3* LivePreviewMesh = Targets[SelectionTargetIndex]->AppliedCanonical.Get();
@@ -230,6 +232,11 @@ void UUVSelectTool::OnSelectionChanged()
 						MovingVids.Add(TriVids[i]);
 					}
 				}
+				if (!TidSet.Contains(Tid))
+				{
+					TidSet.Add(Tid);
+					SelectedTids.Add(Tid);
+				}
 
 				// Gather the boundary edges in the live preview
 				FIndex3i TriEids = LivePreviewMesh->GetTriEdges(Tid);
@@ -259,6 +266,17 @@ void UUVSelectTool::OnSelectionChanged()
 						VidSet.Add(EdgeVids[i]);
 						MovingVids.Add(EdgeVids[i]);
 					}
+
+					TArray<int> TidOneRing;
+					Selection.Mesh->GetVertexOneRingTriangles(EdgeVids[i], TidOneRing);
+					for (int32 Tid : TidOneRing)
+					{
+						if (!TidSet.Contains(Tid))
+						{
+							TidSet.Add(Tid);
+							SelectedTids.Add(Tid);
+						}
+					}
 				}
 			}
 		}
@@ -327,33 +345,8 @@ void UUVSelectTool::GizmoTransformChanged(UTransformProxy* Proxy, FTransform Tra
 
 	if (!DeltaTransform.GetTranslation().IsNearlyZero() || !DeltaTransform.GetRotation().IsIdentity() || Transform.GetScale3D() != FVector::One())
 	{
-		FTransform3d TransformToApply(Transform);
-
-		// TODO: The division here is a bit of a hack. Properly-speaking, the scaling handles should act relative to
-		// gizmo size, not the visible space across which we drag, otherwise it becomes dependent on the units we
-		// use and our absolute distance from the object. Since our UV unwrap is scaled by 1000 to make it
-		// easier to zoom in and out without running into issues, the measure of the distance across which we typically
-		// drag the handles is too high to be convenient. Until we make the scaling invariant to units/distance from
-		// target, we use this hack.
-		TransformToApply.SetScale(FVector::One() + (Transform.GetScale3D() - FVector::One()) / 10);
-
-		const FDynamicMeshSelection& Selection = SelectionMechanic->GetCurrentSelection();
-
-		Targets[SelectionTargetIndex]->UnwrapPreview->PreviewMesh->DeferredEditMesh([&TransformToApply,  &Transform, this](FDynamicMesh3& MeshIn)
-		{
-			for (int32 i = 0; i < MovingVids.Num(); ++i)
-			{
-				MeshIn.SetVertex(MovingVids[i], TransformToApply.TransformPosition(MovingVertOriginalPositions[i]));
-			}
-		}, false);
-		Targets[SelectionTargetIndex]->UpdateUnwrapPreviewOverlay(&MovingVids);
-
-		SelectionMechanic->SetDrawnElementsTransform((FTransform)TransformToApply);
-
-		if (Settings->bUpdatePreviewDuringDrag)
-		{
-			Targets[SelectionTargetIndex]->UpdateAppliedPreviewFromUnwrapPreview(&MovingVids);
-		}
+		UnappliedGizmoTransform = Transform;
+		bGizmoTransformNeedsApplication = true;
 	}	
 }
 
@@ -364,6 +357,9 @@ void UUVSelectTool::GizmoTransformEnded(UTransformProxy* Proxy)
 	//ChangeTracker.BeginChange();
 	//ChangeTracker.SaveTriangles(MovingVids, true);
 
+	// One final attempt to apply transforms if OnTick hasn't happened yet
+	ApplyGizmoTransform();
+
 	if (Settings->bUpdatePreviewDuringDrag)
 	{
 		// Both previews must already be updated, so only need to update canonical
@@ -371,7 +367,7 @@ void UUVSelectTool::GizmoTransformEnded(UTransformProxy* Proxy)
 	}
 	else
 	{
-		Targets[SelectionTargetIndex]->UpdateAllFromUnwrapPreview(&MovingVids);
+		Targets[SelectionTargetIndex]->UpdateAllFromUnwrapPreview(&MovingVids, nullptr, &SelectedTids);
 	}
 
 	if (!AABBTrees[SelectionTargetIndex]->IsValid())
@@ -383,6 +379,40 @@ void UUVSelectTool::GizmoTransformEnded(UTransformProxy* Proxy)
 	SelectionMechanic->RebuildDrawnElements(TransformGizmo->GetGizmoTransform());
 }
 
+void UUVSelectTool::ApplyGizmoTransform()
+{
+	if (bGizmoTransformNeedsApplication)
+	{
+		FTransform3d TransformToApply(UnappliedGizmoTransform);
+
+		// TODO: The division here is a bit of a hack. Properly-speaking, the scaling handles should act relative to
+		// gizmo size, not the visible space across which we drag, otherwise it becomes dependent on the units we
+		// use and our absolute distance from the object. Since our UV unwrap is scaled by 1000 to make it
+		// easier to zoom in and out without running into issues, the measure of the distance across which we typically
+		// drag the handles is too high to be convenient. Until we make the scaling invariant to units/distance from
+		// target, we use this hack.
+		TransformToApply.SetScale(FVector::One() + (UnappliedGizmoTransform.GetScale3D() - FVector::One()) / 10);
+
+		Targets[SelectionTargetIndex]->UnwrapPreview->PreviewMesh->DeferredEditMesh([&TransformToApply, this](FDynamicMesh3& MeshIn)
+			{
+				for (int32 i = 0; i < MovingVids.Num(); ++i)
+				{
+					MeshIn.SetVertex(MovingVids[i], TransformToApply.TransformPosition(MovingVertOriginalPositions[i]));
+				}
+			}, false);
+		Targets[SelectionTargetIndex]->UpdateUnwrapPreviewOverlay(&MovingVids, nullptr, &SelectedTids);
+
+		SelectionMechanic->SetDrawnElementsTransform((FTransform)TransformToApply);
+
+		if (Settings->bUpdatePreviewDuringDrag)
+		{
+			Targets[SelectionTargetIndex]->UpdateAppliedPreviewFromUnwrapPreview(&MovingVids, nullptr, &SelectedTids);
+		}
+
+		bGizmoTransformNeedsApplication = false;
+	}
+}
+
 void UUVSelectTool::Render(IToolsContextRenderAPI* RenderAPI)
 {
 	SelectionMechanic->Render(RenderAPI);
@@ -390,6 +420,7 @@ void UUVSelectTool::Render(IToolsContextRenderAPI* RenderAPI)
 
 void UUVSelectTool::OnTick(float DeltaTime)
 {
+	ApplyGizmoTransform();
 }
 
 
diff --git a/Engine/Plugins/Experimental/UVEditor/Source/UVEditorTools/Public/ToolTargets/UVEditorToolMeshInput.h b/Engine/Plugins/Experimental/UVEditor/Source/UVEditorTools/Public/ToolTargets/UVEditorToolMeshInput.h
--- a/Engine/Plugins/Experimental/UVEditor/Source/UVEditorTools/Public/ToolTargets/UVEditorToolMeshInput.h
+++ b/Engine/Plugins/Experimental/UVEditor/Source/UVEditorTools/Public/ToolTargets/UVEditorToolMeshInput.h
@@ -95,67 +95,70 @@ public:
 	void Shutdown();
 
 	// Notes above the below convenience functions:
-	// 1. If BOTH ChangedVids/ChagnedElementIDs and ChangedTids are null, a complete update will be performed. If either
+	// 1. If BOTH ChangedVids/ChagnedElementIDs and ChangedConnectivityTids are null, a complete update will be performed. If either
 	//   is present, then only those vids/elements/tids will be updated.
 	// 2. Generally ChangedVids/ChangedElementIDs is allowed to have new elements, since it is natural to split UVs.
-	//   However, ChangedTids must not have new Tids, because the tids form our correspondence between the unwrap mesh and
+	//   However, ChangedConnectivityTids must not have new Tids, because the tids form our correspondence between the unwrap mesh and
 	//   the original layer.
-	// 3. If updating the preview objects, note that the functions do not try to cancel any active computations, so an active
+	// 3. FastRenderUpdateTids is an optional list of triangles whose render data needs updating, which can allow for faster preview
+	//  updates if the render buffers are properly split apart. If provided, it should be a superset of ChangedConnectivityTids,
+	//  and should contain the one-ring triangles of ChangedVids.
+	// 4. If updating the preview objects, note that the functions do not try to cancel any active computations, so an active
 	//   computation could reset things once it completes.
 
 	/**
 	 * Updates UnwrapPreview UV Overlay from UnwrapPreview vert positions. Issues a NotifyDeferredEditCompleted
 	 * for both positions and UVs.
 	 */
-	void UpdateUnwrapPreviewOverlay(const TArray<int32>* ChangedVids = nullptr, const TArray<int32>* ChangedTids = nullptr);
+	void UpdateUnwrapPreviewOverlay(const TArray<int32>* ChangedVids = nullptr, const TArray<int32>* ChangedConnectivityTids = nullptr, const TArray<int32>* FastRenderUpdateTids = nullptr);
 
 	/**
 	 * Updates UnwrapCanonical UV Overlay from UnwrapCanonical vert positions. Issues a NotifyDeferredEditCompleted
 	 * for both positions and UVs.
 	 */
-	void UpdateUnwrapCanonicalOverlay(const TArray<int32>* ChangedVids = nullptr, const TArray<int32>* ChangedTids = nullptr);
+	void UpdateUnwrapCanonicalOverlay(const TArray<int32>* ChangedVids = nullptr, const TArray<int32>* ChangedConnectivityTids = nullptr);
 
 	/**
 	 * Updates the AppliedPreview from UnwrapPreview, without updating the non-preview meshes. Useful for updates during
 	 * a drag, etc, when we only care about updating the visible items.
 	 * Assumes that the overlay in UnwrapPreview is already updated.
 	 */
-	void UpdateAppliedPreviewFromUnwrapPreview(const TArray<int32>* ChangedVids = nullptr, const TArray<int32>* ChangedTids = nullptr);
+	void UpdateAppliedPreviewFromUnwrapPreview(const TArray<int32>* ChangedVids = nullptr, const TArray<int32>* ChangedConnectivityTids = nullptr, const TArray<int32>* FastRenderUpdateTids = nullptr);
 
 	/**
 	 * Updates only the UnwrapPreview from AppliedPreview, without updating the non-preview meshes. Useful for updates during
 	 * a drag, etc, when we only care about updating the visible items.
 	 */
 	void UpdateUnwrapPreviewFromAppliedPreview(
-		const TArray<int32>* ChangedElementIDs = nullptr, const TArray<int32>* ChangedTids = nullptr);
+		const TArray<int32>* ChangedElementIDs = nullptr, const TArray<int32>* ChangedConnectivityTids = nullptr, const TArray<int32>* FastRenderUpdateTids = nullptr);
 
 	/**
 	 * Updates the non-preview meshes from their preview counterparts. Useful, for instance, after the completion of a drag
 	 * to update the canonical objects.
 	 */
-	void UpdateCanonicalFromPreviews(const TArray<int32>* ChangedVids = nullptr, const TArray<int32>* ChangedTids = nullptr);
+	void UpdateCanonicalFromPreviews(const TArray<int32>* ChangedVids = nullptr, const TArray<int32>* ChangedConnectivityTids = nullptr);
 
 	/**
 	 * Updates the other meshes using the mesh in UnwrapPreview. Assumes that the overlay in UnwrapPreview is updated.
 	 */
-	void UpdateAllFromUnwrapPreview(const TArray<int32>* ChangedVids = nullptr, const TArray<int32>* ChangedTids = nullptr);
+	void UpdateAllFromUnwrapPreview(const TArray<int32>* ChangedVids = nullptr, const TArray<int32>* ChangedConnectivityTids = nullptr, const TArray<int32>* FastRenderUpdateTids = nullptr);
 
 	/**
 	 * Updates the other meshes using the mesh in UnwrapCanonical. Assumes that overlay in UnwrapCanonical is updated.
 	 */
-	void UpdateAllFromUnwrapCanonical(const TArray<int32>* ChangedVids = nullptr, const TArray<int32>* ChangedTids = nullptr);
+	void UpdateAllFromUnwrapCanonical(const TArray<int32>* ChangedVids = nullptr, const TArray<int32>* ChangedConnectivityTids = nullptr, const TArray<int32>* FastRenderUpdateTids = nullptr);
 
 	/**
 	 * Updates the other meshes using the UV overlay in the live preview.
 	 */
-	void UpdateAllFromAppliedPreview(const TArray<int32>* ChangedElementIDs = nullptr, const TArray<int32>* ChangedTids = nullptr);
+	void UpdateAllFromAppliedPreview(const TArray<int32>* ChangedElementIDs = nullptr, const TArray<int32>* ChangedConnectivityTids = nullptr, const TArray<int32>* FastRenderUpdateTids = nullptr);
 
 	// UToolTarget
 	virtual bool IsValid() const override;
 protected:
 
 	void UpdateAppliedOverlays(const UE::Geometry::FDynamicMeshUVOverlay& SourceOverlay,
-		const TArray<int32>* ChangedVids, const TArray<int32>* ChangedTids);
+		const TArray<int32>* ChangedVids, const TArray<int32>* ChangedConnectivityTids, const TArray<int32>* FastRenderUpdateTids);
 	void UpdateOtherUnwrap(const UE::Geometry::FDynamicMesh3& SourceUnwrapMesh,
-		UE::Geometry::FDynamicMesh3& DestUnwrapMesh, const TArray<int32>* ChangedVids, const TArray<int32>* ChangedTids);
+		UE::Geometry::FDynamicMesh3& DestUnwrapMesh, const TArray<int32>* ChangedVids, const TArray<int32>* ChangedConnectivityTids);
 };
\ No newline at end of file
diff --git a/Engine/Plugins/Experimental/UVEditor/Source/UVEditorTools/Public/UVSelectTool.h b/Engine/Plugins/Experimental/UVEditor/Source/UVEditorTools/Public/UVSelectTool.h
--- a/Engine/Plugins/Experimental/UVEditor/Source/UVEditorTools/Public/UVSelectTool.h
+++ b/Engine/Plugins/Experimental/UVEditor/Source/UVEditorTools/Public/UVSelectTool.h
@@ -63,7 +63,6 @@ public:
 	//~ TODO: Make this only visible in transform mode
 	UPROPERTY(EditAnywhere, Category = Options, AdvancedDisplay)
 	bool bUpdatePreviewDuringDrag = true;
-
 };
 
 /**
@@ -108,6 +107,8 @@ protected:
 	virtual void GizmoTransformStarted(UTransformProxy* Proxy);
 	virtual void GizmoTransformEnded(UTransformProxy* Proxy);
 
+
+	virtual void ApplyGizmoTransform();
 	virtual void UpdateGizmo();
 	virtual void UpdateLivePreviewLines();
 
@@ -134,7 +135,11 @@ protected:
 	ULineSetComponent* LivePreviewLineSet = nullptr;
 
 	UE::Geometry::FFrame3d InitialGizmoFrame;
+	FTransform UnappliedGizmoTransform;
+	bool bGizmoTransformNeedsApplication = false;
+
 	TArray<int32> MovingVids;
+	TArray<int32> SelectedTids;
 	TArray<FVector3d> MovingVertOriginalPositions;
 	int32 SelectionTargetIndex;
 	TArray<int32> BoundaryEids;
