UVEditor - Fix bug where the Alt key held down with mouse clicks would cause cursor to vanish

#rnx

diff --git a/Engine/Plugins/Experimental/UVEditor/Source/UVEditor/Private/UVEditor2DViewportClient.cpp b/Engine/Plugins/Experimental/UVEditor/Source/UVEditor/Private/UVEditor2DViewportClient.cpp
--- a/Engine/Plugins/Experimental/UVEditor/Source/UVEditor/Private/UVEditor2DViewportClient.cpp
+++ b/Engine/Plugins/Experimental/UVEditor/Source/UVEditor/Private/UVEditor2DViewportClient.cpp
@@ -131,3 +131,7 @@ void FUVEditor2DViewportClient::Draw(const FSceneView* View, FPrimitiveDrawInter
 	FEditorViewportClient::Draw(View, PDI);
 }
 
+bool FUVEditor2DViewportClient::ShouldOrbitCamera() const
+{
+	return false; // The UV Editor's 2D viewport should never orbit.
+}
diff --git a/Engine/Plugins/Experimental/UVEditor/Source/UVEditor/Public/UVEditor2DViewportClient.h b/Engine/Plugins/Experimental/UVEditor/Source/UVEditor/Public/UVEditor2DViewportClient.h
--- a/Engine/Plugins/Experimental/UVEditor/Source/UVEditor/Public/UVEditor2DViewportClient.h
+++ b/Engine/Plugins/Experimental/UVEditor/Source/UVEditor/Public/UVEditor2DViewportClient.h
@@ -25,6 +25,7 @@ public:
 	// FEditorViewportClient
 	virtual bool InputKey(FViewport* InViewport, int32 ControllerId, FKey Key, EInputEvent Event, float/*AmountDepressed*/, bool/*Gamepad*/) override;
 	virtual void Draw(const FSceneView* View, FPrimitiveDrawInterface* PDI) override;
+	virtual bool ShouldOrbitCamera() const override;
 
 	// IInputBehaviorSource
 	virtual const UInputBehaviorSet* GetInputBehaviors() const override;
