UVEditor: Add support for textured backgrounds behind UV islands

#rnx

diff --git a/Engine/Plugins/Experimental/UVEditor/Source/UVEditor/Private/UVEditor2DViewportClient.cpp b/Engine/Plugins/Experimental/UVEditor/Source/UVEditor/Private/UVEditor2DViewportClient.cpp
--- a/Engine/Plugins/Experimental/UVEditor/Source/UVEditor/Private/UVEditor2DViewportClient.cpp
+++ b/Engine/Plugins/Experimental/UVEditor/Source/UVEditor/Private/UVEditor2DViewportClient.cpp
@@ -97,7 +97,7 @@ void FUVEditor2DViewportClient::Draw(const FSceneView* View, FPrimitiveDrawInter
 	FVector WorldBoundsMax = View->ScreenToWorld(MaxScreen);
 	FVector WorldBoundsMin = View->ScreenToWorld(MinScreen);
 	FVector ViewLoc = GetViewLocation();
-	ViewLoc.Z = 0.0; // We are treating the scene like a 2D plane, so we'll clamp the Z position here to 
+	ViewLoc.Z = 0.01; // We are treating the scene like a 2D plane, so we'll clamp the Z position here to 
 	               // 0 as a simple projection step just in case.
 
 	// Prevent grid from drawing if we are too close or too far, in order to avoid potential graphical issues.
@@ -122,8 +122,8 @@ void FUVEditor2DViewportClient::Draw(const FSceneView* View, FPrimitiveDrawInter
 	float AxisExtent = FMathf::Max(UVScale, FMathf::Min(WorldBoundsMax.Y, WorldBoundsMax.X));
 
 	// Draw colored axis lines
-	PDI->DrawLine(FVector(0, 0, 0), FVector(AxisExtent, 0, 0), FLinearColor::Red, SDPG_World, AxisThickness, 0, true);
-	PDI->DrawLine(FVector(0, 0, 0), FVector(0, AxisExtent, 0), FLinearColor::Green, SDPG_World, AxisThickness, 0, true);
+	PDI->DrawLine(FVector(0, 0, 0.01), FVector(AxisExtent, 0, 0), FLinearColor::Red, SDPG_World, AxisThickness, 0, true);
+	PDI->DrawLine(FVector(0, 0, 0.01), FVector(0, AxisExtent, 0), FLinearColor::Green, SDPG_World, AxisThickness, 0, true);
 
 	// TODO: Draw a little UV axis thing in the lower left, like the XYZ things that normal viewports have.
 
diff --git a/Engine/Plugins/Experimental/UVEditor/Source/UVEditor/Private/UVEditorBackgroundPreview.cpp b/Engine/Plugins/Experimental/UVEditor/Source/UVEditor/Private/UVEditorBackgroundPreview.cpp
new file mode 100644
--- /dev/null
+++ b/Engine/Plugins/Experimental/UVEditor/Source/UVEditor/Private/UVEditorBackgroundPreview.cpp
@@ -0,0 +1,102 @@
+// Copyright Epic Games, Inc. All Rights Reserved.
+
+#include "UVEditorBackgroundPreview.h"
+#include "DynamicMesh/DynamicMesh3.h"
+#include "DynamicMesh/DynamicMeshAttributeSet.h"
+#include "Drawing/LineSetComponent.h"
+#include "Drawing/MeshWireframeComponent.h"
+#include "ToolSetupUtil.h"
+#include "Async/Async.h"
+
+using namespace UE::Geometry;
+
+void UUVEditorBackgroundPreview::OnCreated()
+{
+	Settings = NewObject<UUVEditorBackgroundPreviewProperties>(this);
+	Settings->WatchProperty(Settings->bVisible, [this](bool) { bSettingsModified = true; });
+	Settings->WatchProperty(Settings->bUseMaterials, [this](bool) { bSettingsModified = true; });
+	Settings->WatchProperty(Settings->SourceTexture, [this](UTexture2D*) {bSettingsModified = true; });
+	Settings->WatchProperty(Settings->SourceMaterial, [this](UMaterial*) {bSettingsModified = true; });
+
+	bSettingsModified = false;
+
+	BackgroundComponent = NewObject<UTriangleSetComponent>(GetActor());
+	BackgroundComponent->SetupAttachment(GetActor()->GetRootComponent());
+	BackgroundComponent->RegisterComponent();
+}
+
+void UUVEditorBackgroundPreview::OnTick(float DeltaTime)
+{
+	if (bSettingsModified)
+	{
+		UpdateBackground();
+		UpdateVisibility();
+		bSettingsModified = false;
+	}
+}
+
+void UUVEditorBackgroundPreview::UpdateVisibility()
+{
+	if (Settings->bVisible == false)
+	{
+		BackgroundComponent->SetVisibility(false);
+		return;
+	}
+
+	BackgroundComponent->SetVisibility(true);
+	BackgroundComponent->MarkRenderStateDirty();
+}
+
+void UUVEditorBackgroundPreview::UpdateBackground()
+{
+
+	const int32 GridCellCountX = 1;
+	const int32 GridCellCountY = 1;
+	const FVector Origin = { 0,0,0 };
+	const FVector GridDx = { 1000, 1000, 0 };
+	const FVector Normal(0, 0, 1);
+	const FColor BackgroundColor = FColor::Blue;
+
+	UMaterial* Material = LoadObject<UMaterial>(nullptr, TEXT("/UVEditor/Materials/UVEditorBackground"));
+	check(Material);	
+	BackgroundMaterial = UMaterialInstanceDynamic::Create(Material, this);
+	if (Settings->bUseMaterials)
+	{
+		if (Settings->SourceMaterial)
+		{
+			BackgroundMaterial = UMaterialInstanceDynamic::Create(Settings->SourceMaterial.Get(), this);
+		}
+	}
+	else
+	{
+		if (Settings->SourceTexture)
+		{
+			BackgroundMaterial->SetTextureParameterValue(TEXT("BackgroundBaseMap"), Settings->SourceTexture);
+		}
+	}
+	BackgroundMaterial->SetScalarParameterValue(TEXT("BackgroundPixelDepthOffset"), 0);
+	BackgroundComponent->Clear();
+
+	for (int32 GridStepX = 0; GridStepX < GridCellCountX; ++GridStepX)
+	{
+		for (int32 GridStepY = 0; GridStepY < GridCellCountY; ++GridStepY)
+		{
+			FVector CellOrigin = Origin + FVector(GridDx.X * GridStepX, GridDx.X * GridStepY, 0);
+			FVector CellOffsetX = { GridDx.X, 0, 0 };
+			FVector CellOffsetY = { 0, GridDx.Y, 0 };
+			
+			FRenderableTriangleVertex A(CellOrigin, { -1,0 }, Normal, BackgroundColor);
+			FRenderableTriangleVertex B(CellOrigin + CellOffsetX, { -1 , -1 }, Normal, BackgroundColor);
+			FRenderableTriangleVertex C(CellOrigin + CellOffsetY, { 0, 0 }, Normal, BackgroundColor);
+			FRenderableTriangleVertex D(CellOrigin + CellOffsetX + CellOffsetY, { 0,-1 }, Normal, BackgroundColor);
+
+			FRenderableTriangle Lower(BackgroundMaterial, A, D, B);
+			FRenderableTriangle Upper(BackgroundMaterial, A, C, D);
+
+			BackgroundComponent->AddTriangle(Lower);
+			BackgroundComponent->AddTriangle(Upper);		
+		}
+	}
+
+	
+}
diff --git a/Engine/Plugins/Experimental/UVEditor/Source/UVEditor/Private/UVEditorBackgroundPreview.h b/Engine/Plugins/Experimental/UVEditor/Source/UVEditor/Private/UVEditorBackgroundPreview.h
new file mode 100644
--- /dev/null
+++ b/Engine/Plugins/Experimental/UVEditor/Source/UVEditor/Private/UVEditorBackgroundPreview.h
@@ -0,0 +1,74 @@
+// Copyright Epic Games, Inc. All Rights Reserved.
+
+#pragma once
+
+#include "CoreMinimal.h"
+#include "Drawing/PreviewGeometryActor.h"
+#include "Drawing/TriangleSetComponent.h"
+#include "InteractiveTool.h"
+#include "GeometryBase.h"
+#include "UVEditorBackgroundPreview.generated.h"
+
+/**
+ * Visualization settings for the UUVEditorBackgroundPreview
+ */
+UCLASS()
+class UVEDITOR_API UUVEditorBackgroundPreviewProperties : public UInteractiveToolPropertySet
+{
+	GENERATED_BODY()
+public:
+	/** Should the background be shown */
+	UPROPERTY(EditAnywhere, Category = Background)
+	bool bVisible = false;
+
+	/** Should the background show textures or materials? */
+	UPROPERTY(EditAnywhere, Category = Background, meta = (EditCondition = "bVisible", EditConditionHides = true))
+	bool bUseMaterials = false;
+
+
+	UPROPERTY(EditAnywhere, Category = Background, meta = (EditCondition = "!bUseMaterials && bVisible", EditConditionHides = true))
+	TObjectPtr<UTexture2D> SourceTexture;
+
+	UPROPERTY(EditAnywhere, Category = Background, meta = (EditCondition = "bUseMaterials && bVisible", EditConditionHides=true))
+	TObjectPtr<UMaterial> SourceMaterial;
+};
+
+/**
+  Serves as a container for the background texture/material display in the UVEditor. This class is responsible for managing the quad
+  drawn behind the grid, as well as maintaining the texture and material choices from the user to display.
+ */
+UCLASS(Transient)
+class UVEDITOR_API UUVEditorBackgroundPreview : public UPreviewGeometry
+{
+	GENERATED_BODY()
+
+public:
+
+	/**
+	 * Client must call this every frame for changes to .Settings to be reflected in rendered result.
+	 */
+	void OnTick(float DeltaTime);
+
+
+public:
+	/** Visualization settings */
+	UPROPERTY()
+		TObjectPtr<UUVEditorBackgroundPreviewProperties> Settings;
+
+	/** The component containing the quad visualization */
+	UPROPERTY()
+		TObjectPtr<UTriangleSetComponent> BackgroundComponent;
+
+	/** The active material being displayed for the background */
+	UPROPERTY()
+		TObjectPtr<UMaterialInstanceDynamic> BackgroundMaterial;
+
+protected:
+	virtual void OnCreated() override;
+
+	bool bSettingsModified = false;
+
+	void UpdateVisibility();
+	void UpdateBackground();
+};
+
diff --git a/Engine/Plugins/Experimental/UVEditor/Source/UVEditor/Private/UVEditorMode.cpp b/Engine/Plugins/Experimental/UVEditor/Source/UVEditor/Private/UVEditorMode.cpp
--- a/Engine/Plugins/Experimental/UVEditor/Source/UVEditor/Private/UVEditorMode.cpp
+++ b/Engine/Plugins/Experimental/UVEditor/Source/UVEditor/Private/UVEditorMode.cpp
@@ -23,6 +23,7 @@
 #include "UVEditorSubsystem.h"
 #include "UVEditorToolUtil.h"
 #include "UVToolContextObjects.h"
+#include "UVEditorBackgroundPreview.h"
 #include "Editor.h"
 
 #define LOCTEXT_NAMESPACE "UUVEditorMode"
@@ -35,11 +36,7 @@ namespace UVEditorModeLocals
 {
 	// TODO: This is temporary. We need to be able to configure the layer we look at,
 	// and ideally show multiple layers at a time.
-	const int32 UVLayerIndex = 0;
-
-	FColor TriangleColor = FColor(50, 194, 219);
-	FColor WireframeColor = FColor(50, 100, 219);
-	FColor IslandBorderColor = FColor(103, 52, 235);
+	const int32 UVLayerIndex = 0;	
 }
 
 const FToolTargetTypeRequirements& UUVEditorMode::GetToolTargetRequirements()
@@ -67,6 +64,11 @@ void UUVEditorMode::Enter()
 {
 	Super::Enter();
 
+	BackgroundVisualization = NewObject<UUVEditorBackgroundPreview>(this);
+	BackgroundVisualization->CreateInWorld(GetWorld(), FTransform::Identity);
+
+	EditorModeToolkit->SetBackgroundSettings(BackgroundVisualization->Settings);
+
 	RegisterTools();
 
 	ActivateDefaultTool();
@@ -89,7 +91,7 @@ void UUVEditorMode::RegisterTools()
 
 void UUVEditorMode::CreateToolkit()
 {
-	Toolkit = MakeShared<FUVEditorModeToolkit>();
+	Toolkit = EditorModeToolkit = MakeShared<FUVEditorModeToolkit>();
 }
 
 void UUVEditorMode::ActivateDefaultTool()
@@ -112,6 +114,11 @@ void UUVEditorMode::Exit()
 	}
 	ToolInputObjects.Reset();
 
+	if (BackgroundVisualization)
+	{
+		BackgroundVisualization->Disconnect();
+	}
+
 	Super::Exit();
 }
 
@@ -170,7 +177,8 @@ void UUVEditorMode::InitializeTargets(TArray<TObjectPtr<UObject>>& AssetsIn, TAr
 			0, ToolSetupUtil::GetCustomTwoSidedDepthOffsetMaterial(
 				GetToolManager(),
 				(FLinearColor)TriangleColor,
-				0)); //depth offset
+				2,  //depth offset
+				TriangleOpacity));
 
 		// Initialize our timestamp so we can later detect changes
 		MeshChangeStamps.Add(ToolInputObject->UnwrapCanonical->GetShapeTimestamp());
@@ -179,7 +187,7 @@ void UUVEditorMode::InitializeTargets(TArray<TObjectPtr<UObject>>& AssetsIn, TAr
 		UMeshElementsVisualizer* WireframeDisplay = NewObject<UMeshElementsVisualizer>(this);
 		WireframeDisplay->CreateInWorld(GetWorld(), FTransform::Identity);
 
-		WireframeDisplay->Settings->DepthBias = 0.5;
+		WireframeDisplay->Settings->DepthBias = 2.5;
 		WireframeDisplay->Settings->bAdjustDepthBiasUsingMeshSize = false;
 		WireframeDisplay->Settings->bShowWireframe = true;
 		WireframeDisplay->Settings->bShowBorders = true;
@@ -291,6 +299,44 @@ void UUVEditorMode::ModeTick(float DeltaTime)
 			WireframeDisplay->OnTick(DeltaTime);
 		}
 	}
+
+	if (BackgroundVisualization)
+	{
+		BackgroundVisualization->OnTick(DeltaTime);
+		if(BackgroundVisualization->Settings->IsPropertySetEnabled())
+		{
+			BackgroundVisualization->Settings->CheckAndUpdateWatched();
+		}
+		else
+		{
+			BackgroundVisualization->Settings->SilentUpdateWatched();
+		}
+
+        // We adjust the mesh opacity depending on whether we're layered over the background or not.
+		if (BackgroundVisualization->Settings->bVisible)
+		{
+			TriangleOpacity = 0.25;
+		}
+		else
+		{
+			TriangleOpacity = 1.0;
+		}
+
+        // Modify the material of the unwrapped mesh to account for the presence/absence of the background, 
+        // changing the opacity as set just above.
+		for (UUVEditorToolMeshInput* ToolInputObject : ToolInputObjects)
+		{
+			ToolInputObject->UnwrapPreview->PreviewMesh->SetMaterial(
+				0, ToolSetupUtil::GetCustomTwoSidedDepthOffsetMaterial(
+					GetToolManager(),
+					(FLinearColor)TriangleColor,
+					2,  //depth offset
+					TriangleOpacity));
+		}
+		
+	}
+
+
 }
 
 #undef LOCTEXT_NAMESPACE
diff --git a/Engine/Plugins/Experimental/UVEditor/Source/UVEditor/Private/UVEditorModeToolkit.cpp b/Engine/Plugins/Experimental/UVEditor/Source/UVEditor/Private/UVEditorModeToolkit.cpp
--- a/Engine/Plugins/Experimental/UVEditor/Source/UVEditor/Private/UVEditorModeToolkit.cpp
+++ b/Engine/Plugins/Experimental/UVEditor/Source/UVEditor/Private/UVEditorModeToolkit.cpp
@@ -11,6 +11,10 @@
 #include "UVEditorCommands.h"
 #include "Widgets/Layout/SBorder.h"
 #include "Widgets/Text/STextBlock.h"
+#include "UVEditorMode.h"
+#include "UVEditorBackgroundPreview.h"
+
+#include "EdModeInteractiveToolsContext.h"
 
 #define LOCTEXT_NAMESPACE "FUVEditorModeToolkit"
 
@@ -26,6 +30,16 @@ FUVEditorModeToolkit::FUVEditorModeToolkit()
 		.Padding(4)
 		[
 			SNew(SVerticalBox)
+
+			+ SVerticalBox::Slot()
+			.AutoHeight()
+			[
+				SAssignNew(ToolMessageArea, STextBlock)
+				.AutoWrapText(true)
+			.Font(FCoreStyle::GetDefaultFontStyle("Bold", 9))
+			.Text(LOCTEXT("UVEditorToolsLabel", "UV Editor Tools"))
+			]
+
 			+ SVerticalBox::Slot()
 			.AutoHeight()
 			.HAlign(HAlign_Left)
@@ -37,7 +51,7 @@ FUVEditorModeToolkit::FUVEditorModeToolkit()
 			]
 
 			+ SVerticalBox::Slot()
-				.FillHeight(1.0f)
+				.AutoHeight()
 				.Padding(2, 0, 0, 0)
 			[
 				SNew(SVerticalBox)
@@ -67,6 +81,39 @@ FUVEditorModeToolkit::FUVEditorModeToolkit()
 					.Text(FText::GetEmpty())
 				]
 			]
+
+            // Todo: Move this out of here once we figure out how to add menus/UI to the viewport
+			+ SVerticalBox::Slot()
+				.AutoHeight()
+				.Padding(2, 0, 0, 0)
+			[
+					SNew(SVerticalBox)
+
+					+ SVerticalBox::Slot()
+					.AutoHeight()
+					[
+						SAssignNew(ToolMessageArea, STextBlock)
+						.AutoWrapText(true)
+						.Font(FCoreStyle::GetDefaultFontStyle("Bold", 9))
+						.Text(LOCTEXT("UVEditorSettingsLabel", "UV Editor Settings"))
+					]
+
+					+ SVerticalBox::Slot()
+					.AutoHeight()
+					[
+						SAssignNew(EditorDetailsContainer, SBorder)
+						.BorderImage(FEditorStyle::GetBrush("NoBorder"))
+					]
+
+					+ SVerticalBox::Slot()
+					.AutoHeight()
+					[
+						SAssignNew(BackgroundDetailsContainer, SBorder)
+						.BorderImage(FEditorStyle::GetBrush("NoBorder"))
+					]
+
+			]
+
 		];
 }
 
@@ -94,6 +141,29 @@ void FUVEditorModeToolkit::Init(const TSharedPtr<IToolkitHost>& InitToolkitHost,
 
 	// Hook up the tool detail panel
 	ToolDetailsContainer->SetContent(DetailsView.ToSharedRef());
+
+	// Hook up the editor detail panel
+	EditorDetailsContainer->SetContent(ModeDetailsView.ToSharedRef());
+
+	// Hook up the background detail panel
+	{
+		FPropertyEditorModule& PropertyEditorModule = FModuleManager::GetModuleChecked<FPropertyEditorModule>("PropertyEditor");
+
+		FDetailsViewArgs BackgroundDetailsViewArgs;
+		BackgroundDetailsViewArgs.bAllowSearch = false;
+		BackgroundDetailsViewArgs.NameAreaSettings = FDetailsViewArgs::HideNameArea;
+		BackgroundDetailsViewArgs.bHideSelectionTip = true;
+		BackgroundDetailsViewArgs.DefaultsOnlyVisibility = EEditDefaultsOnlyNodeVisibility::Automatic;
+		BackgroundDetailsViewArgs.bShowOptions = false;
+		BackgroundDetailsViewArgs.bAllowMultipleTopLevelObjects = true;
+
+		CustomizeDetailsViewArgs(BackgroundDetailsViewArgs);		// allow subclass to customize arguments
+
+		BackgroundDetailsView = PropertyEditorModule.CreateDetailView(BackgroundDetailsViewArgs);
+		UUVEditorMode* UVEditorMode = Cast<UUVEditorMode>(OwningEditorMode.Get());
+		BackgroundDetailsContainer->SetContent(BackgroundDetailsView.ToSharedRef());
+	}
+	
 }
 
 FName FUVEditorModeToolkit::GetToolkitFName() const
@@ -106,4 +176,11 @@ FText FUVEditorModeToolkit::GetBaseToolkitName() const
 	return NSLOCTEXT("UVEditorModeToolkit", "DisplayName", "UVEditorMode");
 }
 
+
+void FUVEditorModeToolkit::SetBackgroundSettings(UObject* InSettingsObject)
+{
+	BackgroundDetailsView->SetObject(InSettingsObject);
+}
+
+
 #undef LOCTEXT_NAMESPACE
diff --git a/Engine/Plugins/Experimental/UVEditor/Source/UVEditor/Public/UVEditorMode.h b/Engine/Plugins/Experimental/UVEditor/Source/UVEditor/Public/UVEditorMode.h
--- a/Engine/Plugins/Experimental/UVEditor/Source/UVEditor/Public/UVEditorMode.h
+++ b/Engine/Plugins/Experimental/UVEditor/Source/UVEditor/Public/UVEditorMode.h
@@ -15,9 +15,11 @@ class UMeshElementsVisualizer;
 class UPreviewMesh;
 class UToolTarget;
 class UWorld;
+class FUVEditorModeToolkit;
 class UInteractiveToolPropertySet;
 class UUVToolStateObjectStore;
 class UUVEditorToolMeshInput;
+class UUVEditorBackgroundPreview;
 
 /**
  * The UV editor mode is the mode used in the UV asset editor. It holds most of the inter-tool state.
@@ -72,6 +74,12 @@ public:
 	// We don't actually override MouseEnter, etc, because things get forwarded to the input
 	// router via FEditorModeTools, and we don't have any additional input handling to do at the mode level.
 
+
+	// Holds the background visualiztion
+	UPROPERTY()
+		TObjectPtr<UUVEditorBackgroundPreview> BackgroundVisualization;
+
+
 protected:
 
 	// UEdMode overrides
@@ -82,6 +90,17 @@ protected:
 	virtual void OnToolStarted(UInteractiveToolManager* Manager, UInteractiveTool* Tool) {}
 	virtual void OnToolEnded(UInteractiveToolManager* Manager, UInteractiveTool* Tool) {}
 
+	// Wireframe Display Properties
+	float TriangleOpacity = 1.0;
+	FColor TriangleColor = FColor(50, 194, 219);
+	FColor WireframeColor = FColor(50, 100, 219);
+	FColor IslandBorderColor = FColor(103, 52, 235);
+
+    /**
+     * Stores a pointer to the specific EditorModeToolkit so we can interact with specific UI detail views
+     */
+	TSharedPtr<FUVEditorModeToolkit> EditorModeToolkit;
+
 	/**
 	 * This stores the original objects that the UV editor was asked to operate on (for
 	 * instance UStaticMesh pointers).
diff --git a/Engine/Plugins/Experimental/UVEditor/Source/UVEditor/Public/UVEditorModeToolkit.h b/Engine/Plugins/Experimental/UVEditor/Source/UVEditor/Public/UVEditorModeToolkit.h
--- a/Engine/Plugins/Experimental/UVEditor/Source/UVEditor/Public/UVEditorModeToolkit.h
+++ b/Engine/Plugins/Experimental/UVEditor/Source/UVEditor/Public/UVEditorModeToolkit.h
@@ -5,6 +5,7 @@
 #include "CoreMinimal.h"
 
 #include "Toolkits/BaseToolkit.h"
+#include "InteractiveTool.h"
 
 class SBorder;
 class STextBlock;
@@ -27,6 +28,9 @@ public:
 	virtual FText GetBaseToolkitName() const override;
 	virtual TSharedPtr<class SWidget> GetInlineContent() const override { return ToolkitWidget; }
 
+
+	virtual void SetBackgroundSettings(UObject* InSettingsObject);
+
 protected:
 	// The mode's entire toolbox, which gets returned by GetInlineContent()
 	TSharedPtr<SWidget> ToolkitWidget;
@@ -41,6 +45,13 @@ protected:
 	// in FModeToolkit
 	TSharedPtr<SBorder> ToolDetailsContainer;
 
+	// A container for the editor settings
+	TSharedPtr<SBorder> EditorDetailsContainer;
+
+	// A container for the background settings
+	TSharedPtr<SBorder> BackgroundDetailsContainer;
+	TSharedPtr<IDetailsView> BackgroundDetailsView;
+
 	// A place for tools to write out any instructions
 	TSharedPtr<STextBlock> ToolMessageArea;
 };
\ No newline at end of file
