MeshModelingTools: Support undo/redo in UV editor

#rnx

diff --git a/Engine/Plugins/Experimental/MeshModelingToolsetExp/Source/ModelingComponents/Private/Selection/MeshSelectionMechanic.cpp b/Engine/Plugins/Experimental/MeshModelingToolsetExp/Source/ModelingComponents/Private/Selection/MeshSelectionMechanic.cpp
--- a/Engine/Plugins/Experimental/MeshModelingToolsetExp/Source/ModelingComponents/Private/Selection/MeshSelectionMechanic.cpp
+++ b/Engine/Plugins/Experimental/MeshModelingToolsetExp/Source/ModelingComponents/Private/Selection/MeshSelectionMechanic.cpp
@@ -19,6 +19,53 @@ using namespace UE::Geometry;
 
 #define LOCTEXT_NAMESPACE "UMeshSelectionMechanic"
 
+namespace MeshSelectionMechanicLocals
+{
+	/**
+	 * Class that can be used to apply/revert selection changes. Expects that the mesh selection
+	 * mechanic is the associated UObject that is passed to Apply/Revert.
+	 */
+	class FMeshSelectionMechanicSelectionChange : public FToolCommandChange
+	{
+	public:
+		FMeshSelectionMechanicSelectionChange(const FDynamicMeshSelection& OldSelection,
+			const FDynamicMeshSelection& NewSelection, bool bBroadcastOnSelectionChangedIn)
+			: Before(OldSelection)
+			, After(NewSelection)
+			, bBroadcastOnSelectionChanged(bBroadcastOnSelectionChangedIn)
+		{
+		}
+
+		virtual void Apply(UObject* Object) override
+		{
+			UMeshSelectionMechanic* Mechanic = Cast<UMeshSelectionMechanic>(Object);
+			if (Mechanic)
+			{
+				Mechanic->SetSelection(After, bBroadcastOnSelectionChanged, false);
+			}
+		}
+
+		virtual void Revert(UObject* Object) override
+		{
+			UMeshSelectionMechanic* Mechanic = Cast<UMeshSelectionMechanic>(Object);
+			if (Mechanic)
+			{
+				Mechanic->SetSelection(Before, bBroadcastOnSelectionChanged, false);
+			}
+		}
+
+		virtual FString ToString() const override
+		{
+			return TEXT("MeshSelectionMechanicLocals::FMeshSelectionMechanicSelectionChange");
+		}
+
+	protected:
+		FDynamicMeshSelection Before;
+		FDynamicMeshSelection After;
+		bool bBroadcastOnSelectionChanged;
+	};
+}
+
 void UMeshSelectionMechanic::Setup(UInteractiveTool* ParentToolIn)
 {
 	UInteractionMechanic::Setup(ParentToolIn);
@@ -37,6 +84,16 @@ void UMeshSelectionMechanic::Setup(UInteractiveTool* ParentToolIn)
 	PointSet->SetPointMaterial(ToolSetupUtil::GetDefaultPointComponentMaterial(
 		GetParentTool()->GetToolManager(), /*bDepthTested*/ true));
 
+	// Add default selection change emitter if one was not provided.
+	if (!EmitSelectionChange)
+	{
+		EmitSelectionChange = [this](const FDynamicMeshSelection& OldSelection, const FDynamicMeshSelection& NewSelection, bool bBroadcastOnSelectionChanged)
+		{
+			GetParentTool()->GetToolManager()->EmitObjectChange(this, 
+				MakeUnique<MeshSelectionMechanicLocals::FMeshSelectionMechanicSelectionChange>(OldSelection, NewSelection, bBroadcastOnSelectionChanged),
+				LOCTEXT("SelectionChangeMessage", "Selection Change"));
+		};
+	}
 }
 
 void UMeshSelectionMechanic::SetWorld(UWorld* World)
@@ -98,10 +155,23 @@ const FDynamicMeshSelection& UMeshSelectionMechanic::GetCurrentSelection() const
 
 void UMeshSelectionMechanic::SetSelection(const FDynamicMeshSelection& Selection, bool bBroadcast, bool bEmitChange)
 {
+	FDynamicMeshSelection OriginalSelection = CurrentSelection;
 	CurrentSelection = Selection;
 
-	if (MeshSpatials[CurrentSelectionIndex]->GetMesh() != Selection.Mesh)
+	// Adjust the mesh in which the selection lives
+	if (CurrentSelection.Mesh == nullptr)
 	{
+		// We're actually clearing the selection
+		if (!ensure(CurrentSelection.SelectedIDs.IsEmpty()))
+		{
+			CurrentSelection.SelectedIDs.Empty();
+		}
+		CurrentSelectionIndex = IndexConstants::InvalidID;
+	}
+	else if (CurrentSelectionIndex == IndexConstants::InvalidID 
+		|| MeshSpatials[CurrentSelectionIndex]->GetMesh() != Selection.Mesh)
+	{
+		CurrentSelectionIndex = IndexConstants::InvalidID;
 		for (int32 i = 0; i < MeshSpatials.Num(); ++i)
 		{
 			if (MeshSpatials[i]->GetMesh() == Selection.Mesh)
@@ -110,16 +180,20 @@ void UMeshSelectionMechanic::SetSelection(const FDynamicMeshSelection& Selection
 				break;
 			}
 		}
+		ensure(CurrentSelectionIndex != IndexConstants::InvalidID);
 	}
 
 	UpdateCentroid();
 	RebuildDrawnElements(FTransform(CurrentSelectionCentroid));
 
+	if (bEmitChange && OriginalSelection != CurrentSelection)
+	{
+		EmitSelectionChange(OriginalSelection, CurrentSelection, bBroadcast);
+	}
 	if (bBroadcast)
 	{
 		OnSelectionChanged.Broadcast();
 	}
-	// TODO: Undo/redo
 }
 
 void UMeshSelectionMechanic::RebuildDrawnElements(const FTransform& StartTransform)
@@ -128,6 +202,11 @@ void UMeshSelectionMechanic::RebuildDrawnElements(const FTransform& StartTransfo
 	PointSet->Clear();
 	PreviewGeometryActor->SetActorTransform(StartTransform);
 
+	if (CurrentSelectionIndex == IndexConstants::InvalidID)
+	{
+		return;
+	}
+
 	// For us to end up with the StartTransform, we have to invert it, but only after applying
 	// the mesh transform to begin with.
 	auto TransformToApply = [&StartTransform, this](const FVector& VectorIn)
@@ -210,10 +289,21 @@ void UMeshSelectionMechanic::UpdateCentroid()
 		}
 		CurrentSelectionCentroid /= CurrentSelection.SelectedIDs.Num();
 	}
+
+	CentroidTimestamp = CurrentSelection.Mesh->GetShapeTimestamp();
 }
 
 FVector3d UMeshSelectionMechanic::GetCurrentSelectionCentroid()
 {
+	if (!CurrentSelection.Mesh)
+	{
+		return FVector3d(0);
+	}
+	else if (CurrentSelection.Mesh->GetShapeTimestamp() != CentroidTimestamp)
+	{
+		UpdateCentroid();
+	}
+
 	return CurrentSelectionCentroid;
 }
 
@@ -264,7 +354,9 @@ void UMeshSelectionMechanic::OnClicked(const FInputDeviceRay& ClickPos)
 	for (int32 i = 0; i < MeshSpatials.Num(); ++i)
 	{
 		//Short circuit the loop if we're in multiselect mode, to save us from testing things we aren't allowed to select
-		if (InMultiSelectMode() && !CurrentSelection.SelectedIDs.IsEmpty() && i != CurrentSelectionIndex) 
+		if (InMultiSelectMode() && !CurrentSelection.IsEmpty() 
+			&& CurrentSelectionIndex != IndexConstants::InvalidID 
+			&& i != CurrentSelectionIndex)
 		{
 			continue;
 		}
@@ -279,6 +371,7 @@ void UMeshSelectionMechanic::OnClicked(const FInputDeviceRay& ClickPos)
 		{
 			// It should be okay to reassign this. If we're in multiselect, and it misses a hit, this just won't happen and selection won't be altered.
 			CurrentSelection.Mesh = MeshSpatials[i]->GetMesh();
+			CurrentSelection.TopologyTimestamp = CurrentSelection.Mesh->GetTopologyTimestamp();
 			CurrentSelectionIndex = i;
 			break;
 		}
@@ -369,10 +462,11 @@ void UMeshSelectionMechanic::OnClicked(const FInputDeviceRay& ClickPos)
 		}
 	}
 
-	if (!(OriginalSelection == CurrentSelection))
+	if (OriginalSelection != CurrentSelection)
 	{
 		UpdateCentroid();
-		RebuildDrawnElements(FTransform(CurrentSelectionCentroid));
+		RebuildDrawnElements(FTransform(GetCurrentSelectionCentroid()));
+		EmitSelectionChange(OriginalSelection, CurrentSelection, true);
 		OnSelectionChanged.Broadcast();
 	}
 }
diff --git a/Engine/Plugins/Experimental/MeshModelingToolsetExp/Source/ModelingComponents/Public/Selection/DynamicMeshSelection.h b/Engine/Plugins/Experimental/MeshModelingToolsetExp/Source/ModelingComponents/Public/Selection/DynamicMeshSelection.h
--- a/Engine/Plugins/Experimental/MeshModelingToolsetExp/Source/ModelingComponents/Public/Selection/DynamicMeshSelection.h
+++ b/Engine/Plugins/Experimental/MeshModelingToolsetExp/Source/ModelingComponents/Public/Selection/DynamicMeshSelection.h
@@ -30,6 +30,17 @@ public:
 	// Not relevant if the selection type is not Group
 	int32 GroupLayer = 0;
 
+	// Can be used to discard the selection if topology of the mesh
+	// has changed (to avoid the risk of referencing elements that may
+	// have been deleted)
+	int32 TopologyTimestamp = -1;
+
+	/** Checks whether the selection's timestamp still matches the meshes topology timestamp. */
+	bool MatchesTimestamp() const
+	{
+		return Mesh && Mesh->GetTopologyTimestamp() == TopologyTimestamp;
+	}
+
 	bool IsEmpty() const
 	{
 		return SelectedIDs.IsEmpty();
@@ -39,11 +50,16 @@ public:
 	{
 		return Mesh == Other.Mesh
 			&& Type == Other.Type
+			&& TopologyTimestamp == Other.TopologyTimestamp
 			&& (Type != EType::Group || GroupLayer == Other.GroupLayer)
 			&& SelectedIDs.Num() == Other.SelectedIDs.Num()
 			&& SelectedIDs.Includes(Other.SelectedIDs);
 	}
 
+	bool operator!=(const FDynamicMeshSelection& Other)
+	{
+		return !(*this == Other);
+	}
 };
 
 }}// end namespace UE::Geometry
\ No newline at end of file
diff --git a/Engine/Plugins/Experimental/MeshModelingToolsetExp/Source/ModelingComponents/Public/Selection/MeshSelectionMechanic.h b/Engine/Plugins/Experimental/MeshModelingToolsetExp/Source/ModelingComponents/Public/Selection/MeshSelectionMechanic.h
--- a/Engine/Plugins/Experimental/MeshModelingToolsetExp/Source/ModelingComponents/Public/Selection/MeshSelectionMechanic.h
+++ b/Engine/Plugins/Experimental/MeshModelingToolsetExp/Source/ModelingComponents/Public/Selection/MeshSelectionMechanic.h
@@ -43,7 +43,7 @@ public:
 /**
  * Mechanic for selecting elements of a dynamic mesh.
  * 
- * TODO: Not finished. Needs undo/redo, currently only selects connected components or unoccluded edges.
+ * TODO: Currently only able to select unoccluded elements.
  */
 UCLASS()
 class MODELINGCOMPONENTS_API UMeshSelectionMechanic : public UInteractionMechanic, public IClickBehaviorTarget
@@ -91,6 +91,16 @@ public:
 
 	EMeshSelectionMechanicMode SelectionMode;
 
+	/**
+	 * Function to use for emitting selection change events. If not set, Setup() will attach a version that
+	 * uses EmitObjectChange() on the tool manager to emit a change that operates on this mechanic.
+	 * @param bBroadcastOnSelectionChanged Indicates whether the SetSelection call that asked to issue the
+	 *   transaction was also requested to broadcast the OnSelectionChanged delegate. If so, the transaction
+	 *   should probably do the same on revert/apply.
+	 */ 
+	TUniqueFunction<void(const FDynamicMeshSelection& OldSelection, const FDynamicMeshSelection& NewSelection, 
+		bool bBroadcastOnSelectionChanged)> EmitSelectionChange;
+
 protected:
 
 	bool InMultiSelectMode() const { return bShiftToggle; }
@@ -108,7 +118,7 @@ protected:
 	TArray<TSharedPtr<FDynamicMeshAABBTree3>> MeshSpatials;
 	TArray<FTransform> MeshTransforms;
 	FDynamicMeshSelection CurrentSelection;
-	int32 CurrentSelectionIndex;
+	int32 CurrentSelectionIndex = IndexConstants::InvalidID;
 	FViewCameraState CameraState;
 
 	FColor LineColor = FColor::Yellow;
@@ -119,9 +129,11 @@ protected:
 
 	float LineThickness = 3;
 	float PointThickness = 6;
-	float DepthBias = 0.3;
+	// TODO: This should probably be lower, but we have to lower the depth bias elsewhere first.
+	float DepthBias = 5;
 
 	FVector3d CurrentSelectionCentroid;
+	int32 CentroidTimestamp = -1;
 	void UpdateCentroid();
 };
 
diff --git a/Engine/Plugins/Experimental/UVEditor/Source/UVEditor/Private/UVEditorMode.cpp b/Engine/Plugins/Experimental/UVEditor/Source/UVEditor/Private/UVEditorMode.cpp
--- a/Engine/Plugins/Experimental/UVEditor/Source/UVEditor/Private/UVEditorMode.cpp
+++ b/Engine/Plugins/Experimental/UVEditor/Source/UVEditor/Private/UVEditorMode.cpp
@@ -38,6 +38,8 @@ namespace UVEditorModeLocals
 	// TODO: This is temporary. We need to be able to configure the layer we look at,
 	// and ideally show multiple layers at a time.
 	const int32 UVLayerIndex = 0;	
+
+	const FString DefaultToolIdentifier = TEXT("UVSelectTool");
 }
 
 const FToolTargetTypeRequirements& UUVEditorMode::GetToolTargetRequirements()
@@ -101,9 +103,13 @@ void UUVEditorMode::CreateToolkit()
 
 void UUVEditorMode::ActivateDefaultTool()
 {
-	ToolsContext->StartTool(TEXT("UVSelectTool"));
+	ToolsContext->StartTool(UVEditorModeLocals::DefaultToolIdentifier);
 }
 
+bool UUVEditorMode::IsDefaultToolActive()
+{
+	return ToolsContext->IsToolActive(EToolSide::Mouse, UVEditorModeLocals::DefaultToolIdentifier);
+}
 
 void UUVEditorMode::BindCommands()
 {
diff --git a/Engine/Plugins/Experimental/UVEditor/Source/UVEditor/Private/UVEditorModeToolkit.cpp b/Engine/Plugins/Experimental/UVEditor/Source/UVEditor/Private/UVEditorModeToolkit.cpp
--- a/Engine/Plugins/Experimental/UVEditor/Source/UVEditor/Private/UVEditorModeToolkit.cpp
+++ b/Engine/Plugins/Experimental/UVEditor/Source/UVEditor/Private/UVEditorModeToolkit.cpp
@@ -20,6 +20,46 @@
 
 #define LOCTEXT_NAMESPACE "FUVEditorModeToolkit"
 
+namespace UVEditorModeToolkitLocals
+{
+	/** 
+	 * Support for undoing a tool start in such a way that we go back to the mode's default
+	 * tool on undo.
+	 */
+	class FUVEditorBeginToolChange : public FToolCommandChange
+	{
+	public:
+		virtual void Apply(UObject* Object) override
+		{
+			// Do nothing, since we don't allow a re-do back into a tool
+		}
+
+		virtual void Revert(UObject* Object) override
+		{
+			UUVEditorMode* Mode = Cast<UUVEditorMode>(Object);
+			// Don't really need the check for default tool since we theoretically shouldn't
+			// be issuing this transaction for starting the default tool, but still...
+			if (Mode && !Mode->IsDefaultToolActive())
+			{
+				Mode->GetInteractiveToolsContext()->EndTool(EToolShutdownType::Cancel);
+				Mode->ActivateDefaultTool();
+			}
+		}
+
+		virtual bool HasExpired(UObject* Object) const override
+		{
+			// To not be expired, we must be in some non-default tool.
+			UUVEditorMode* Mode = Cast<UUVEditorMode>(Object);
+			return !(Mode && Mode->GetInteractiveToolsContext()->ToolManager->HasAnyActiveTool() && !Mode->IsDefaultToolActive());
+		}
+
+		virtual FString ToString() const override
+		{
+			return TEXT("FUVEditorBeginToolChange");
+		}
+	};
+}
+
 FUVEditorModeToolkit::FUVEditorModeToolkit()
 {
 	// Construct the panel that we will give in GetInlineContent().
@@ -127,11 +167,18 @@ void FUVEditorModeToolkit::Init(const TSharedPtr<IToolkitHost>& InitToolkitHost,
 {
 	FModeToolkit::Init(InitToolkitHost, InOwningMode);	
 
+	UUVEditorMode* UVEditorMode = Cast<UUVEditorMode>(GetScriptableEditorMode());
+	check(UVEditorMode);
+
+	// Currently, there's no EToolChangeTrackingMode that reverts back to a default tool on undo (if we add that
+	// support, the tool manager will need to be aware of the default tool). So, we instead opt to do our own management
+	// of tool start transactions. See FUVEditorModeToolkit::OnToolStarted for how we issue the transactions.
+	UVEditorMode->GetInteractiveToolsContext()->ToolManager->ConfigureChangeTrackingMode(EToolChangeTrackingMode::NoChangeTracking);
+
 	// Build the tool palette
 	const FUVEditorCommands& CommandInfos = FUVEditorCommands::Get();
 	TSharedRef<FUICommandList> CommandList = GetToolkitCommands();
-	check(OwningEditorMode.IsValid());
-	FUniformToolBarBuilder ToolbarBuilder(CommandList, FMultiBoxCustomization(OwningEditorMode->GetModeInfo().ToolbarCustomizationName), TSharedPtr<FExtender>(), false);
+	FUniformToolBarBuilder ToolbarBuilder(CommandList, FMultiBoxCustomization(UVEditorMode->GetModeInfo().ToolbarCustomizationName), TSharedPtr<FExtender>(), false);
 	ToolbarBuilder.SetStyle(&FEditorStyle::Get(), "PaletteToolBar");
 
 	// TODO: Add more tools here
@@ -163,7 +210,6 @@ void FUVEditorModeToolkit::Init(const TSharedPtr<IToolkitHost>& InitToolkitHost,
 		CustomizeDetailsViewArgs(BackgroundDetailsViewArgs);		// allow subclass to customize arguments
 
 		BackgroundDetailsView = PropertyEditorModule.CreateDetailView(BackgroundDetailsViewArgs);
-		UUVEditorMode* UVEditorMode = Cast<UUVEditorMode>(OwningEditorMode.Get());
 		BackgroundDetailsContainer->SetContent(BackgroundDetailsView.ToSharedRef());
 	}
 	
@@ -285,10 +331,18 @@ void FUVEditorModeToolkit::OnToolStarted(UInteractiveToolManager* Manager, UInte
 
 	if (GetScriptableEditorMode()->GetInteractiveToolsContext()->ActiveToolHasAccept())
 	{
+		// Add the accept/cancel overlay
 		GetToolkitHost()->AddViewportOverlayWidget(ViewportOverlayWidget.ToSharedRef());
 	}
 	
-
+	// Issue a tool start transaction unless we are starting the default tool, because we can't
+	// undo or revert out of the default tool.
+	UUVEditorMode* Mode = Cast<UUVEditorMode>(GetScriptableEditorMode());
+	if (!Mode->IsDefaultToolActive())
+	{
+		Mode->GetInteractiveToolsContext()->GetTransactionAPI()->AppendChange(
+			Mode, MakeUnique<UVEditorModeToolkitLocals::FUVEditorBeginToolChange>(), LOCTEXT("ActivateTool", "Activate Tool"));
+	}
 }
 
 void FUVEditorModeToolkit::OnToolEnded(UInteractiveToolManager* Manager, UInteractiveTool* Tool)
diff --git a/Engine/Plugins/Experimental/UVEditor/Source/UVEditor/Private/UVEditorToolkit.cpp b/Engine/Plugins/Experimental/UVEditor/Source/UVEditor/Private/UVEditorToolkit.cpp
--- a/Engine/Plugins/Experimental/UVEditor/Source/UVEditor/Private/UVEditorToolkit.cpp
+++ b/Engine/Plugins/Experimental/UVEditor/Source/UVEditor/Private/UVEditorToolkit.cpp
@@ -326,6 +326,12 @@ void FUVEditorToolkit::PostInitAssetEditor()
 	LivePreviewAPI->Initialize(LivePreviewScene->GetWorld(), LivePreviewInputRouter);
 	ContextStore->AddContextObject(LivePreviewAPI);
 
+	// It is debatable where the undo/redo api should get initialized, but this seems like a 
+	// reasonable place since we're setting up other things for the mode for the first time.
+	UUVToolEmitChangeAPI* EmitChangeAPI = NewObject<UUVToolEmitChangeAPI>();
+	EmitChangeAPI->Initialize(EditorModeManager->GetInteractiveToolsContext()->ToolManager);
+	ContextStore->AddContextObject(EmitChangeAPI);
+
 	// Initialize mode state.
 	TArray<TObjectPtr<UObject>> ObjectsToEdit;
 	OwningAssetEditor->GetObjectsToEdit(ObjectsToEdit);
diff --git a/Engine/Plugins/Experimental/UVEditor/Source/UVEditor/Public/UVEditorMode.h b/Engine/Plugins/Experimental/UVEditor/Source/UVEditor/Public/UVEditorMode.h
--- a/Engine/Plugins/Experimental/UVEditor/Source/UVEditor/Public/UVEditorMode.h
+++ b/Engine/Plugins/Experimental/UVEditor/Source/UVEditor/Public/UVEditorMode.h
@@ -75,6 +75,8 @@ public:
 	// it to later get (accidentally) more entangled with mode internals. At the same time, we're not
 	// sure whether we want to make the UEdMode one public. So this is the minimal-impact tweak.
 	virtual void ActivateDefaultTool() override;
+	// This is currently not part of the base class at all... Should it be?
+	virtual bool IsDefaultToolActive();
 
 	// We don't actually override MouseEnter, etc, because things get forwarded to the input
 	// router via FEditorModeTools, and we don't have any additional input handling to do at the mode level.
diff --git a/Engine/Plugins/Experimental/UVEditor/Source/UVEditorTools/Private/ToolTargets/UVEditorToolMeshInput.cpp b/Engine/Plugins/Experimental/UVEditor/Source/UVEditorTools/Private/ToolTargets/UVEditorToolMeshInput.cpp
--- a/Engine/Plugins/Experimental/UVEditor/Source/UVEditorTools/Private/ToolTargets/UVEditorToolMeshInput.cpp
+++ b/Engine/Plugins/Experimental/UVEditor/Source/UVEditorTools/Private/ToolTargets/UVEditorToolMeshInput.cpp
@@ -4,6 +4,8 @@
 
 #include "Drawing/MeshElementsVisualizer.h" // for wireframe display
 #include "DynamicMesh/DynamicMesh3.h"
+#include "DynamicMesh/DynamicMeshChangeTracker.h" // FDynamicMeshChange
+#include "DynamicMesh/MeshIndexUtil.h"
 #include "MeshOpPreviewHelpers.h"
 #include "ModelingToolTargetUtil.h"
 #include "TargetInterfaces/AssetBackedTarget.h"
@@ -383,3 +385,21 @@ void UUVEditorToolMeshInput::UpdateOtherUnwrap(const FDynamicMesh3& SourceUnwrap
 	FDynamicMeshUVOverlay* DestOverlay = DestUnwrapMesh.Attributes()->PrimaryUV();
 	CopyMeshOverlay(*SourceOverlay, *DestOverlay, ChangedVids, ChangedConnectivityTids);
 }
+
+void UUVEditorToolMeshInput::UpdateFromCanonicalUnwrapUsingMeshChange(const FDynamicMeshChange& UnwrapCanonicalMeshChange)
+{
+	// Note that we know that no triangles were created or destroyed since the UV editor
+	// does not allow that (it would break the mesh mappings). Otherwise we would need to
+	// combine original and final tris here.
+	TArray<int32> ChangedTids;
+	UnwrapCanonicalMeshChange.GetSavedTriangleList(ChangedTids, true);
+
+	TArray<int32> ChangedVids;
+	TriangleToVertexIDs(UnwrapCanonical.Get(), ChangedTids, ChangedVids);
+
+	TSet<int32> RenderUpdateTidsSet;
+	VertexToTriangleOneRing(UnwrapCanonical.Get(), ChangedVids, RenderUpdateTidsSet);
+	TArray<int32> RenderUpdateTids = RenderUpdateTidsSet.Array();
+
+	UpdateAllFromUnwrapCanonical(&ChangedVids, &ChangedTids, &RenderUpdateTids);
+}
diff --git a/Engine/Plugins/Experimental/UVEditor/Source/UVEditorTools/Private/UVEditorLayoutTool.cpp b/Engine/Plugins/Experimental/UVEditor/Source/UVEditorTools/Private/UVEditorLayoutTool.cpp
--- a/Engine/Plugins/Experimental/UVEditor/Source/UVEditorTools/Private/UVEditorLayoutTool.cpp
+++ b/Engine/Plugins/Experimental/UVEditor/Source/UVEditorTools/Private/UVEditorLayoutTool.cpp
@@ -2,12 +2,17 @@
 
 #include "UVEditorLayoutTool.h"
 
+#include "ContextObjectStore.h"
+#include "DynamicMesh/DynamicMesh3.h"
+#include "DynamicMesh/DynamicMeshChangeTracker.h"
 #include "InteractiveToolManager.h"
 #include "MeshOpPreviewHelpers.h" // UMeshOpPreviewWithBackgroundCompute
 #include "ParameterizationOps/UVLayoutOp.h"
 #include "Properties/UVLayoutProperties.h"
 #include "ToolTargets/UVEditorToolMeshInput.h"
+#include "UVToolContextObjects.h"
 
+using namespace UE::Geometry;
 
 #define LOCTEXT_NAMESPACE "UUVSelectTool"
 
@@ -77,9 +82,26 @@ void UUVEditorLayoutTool::Shutdown(EToolShutdownType ShutdownType)
 
 	if (ShutdownType == EToolShutdownType::Accept)
 	{
+		UUVToolEmitChangeAPI* ChangeAPI = GetToolManager()->GetContextObjectStore()->FindContext<UUVToolEmitChangeAPI>();
+
 		for (TObjectPtr<UUVEditorToolMeshInput> Target : Targets)
 		{
+			// Set things up for undo. 
+			// TODO: It's not entirely clear whether it would be safe to use a FMeshVertexChange instead... It seems like
+			// when bAllowFlips is true, we would end up with changes to the tris of the unwrap. Also, if we stick to saving
+			// all the tris and verts, should we consider using the new dynamic mesh serialization?
+			FDynamicMeshChangeTracker ChangeTracker(Target->UnwrapCanonical.Get());
+			ChangeTracker.BeginChange();
+			
+			for (int32 Tid : Target->UnwrapCanonical->TriangleIndicesItr())
+			{
+				ChangeTracker.SaveTriangle(Tid, true);
+			}
+
+			// TODO: Again, it's not clear whether we need to update the entire triangle topology...
 			Target->UpdateCanonicalFromPreviews();
+
+			ChangeAPI->EmitToolIndependentUnwrapCanonicalChange(Target, ChangeTracker.EndChange(), LOCTEXT("ApplyLayoutTool", "Layout Tool"));
 		}
 	}
 	else
diff --git a/Engine/Plugins/Experimental/UVEditor/Source/UVEditorTools/Private/UVSelectTool.cpp b/Engine/Plugins/Experimental/UVEditor/Source/UVEditorTools/Private/UVSelectTool.cpp
--- a/Engine/Plugins/Experimental/UVEditor/Source/UVEditorTools/Private/UVSelectTool.cpp
+++ b/Engine/Plugins/Experimental/UVEditor/Source/UVEditorTools/Private/UVSelectTool.cpp
@@ -4,6 +4,7 @@
 
 #include "BaseGizmos/TransformGizmo.h"
 #include "ContextObjectStore.h"
+#include "DynamicMesh/DynamicMeshChangeTracker.h"
 #include "Drawing/LineSetComponent.h"
 #include "Drawing/MeshElementsVisualizer.h"
 #include "Drawing/PreviewGeometryActor.h"
@@ -15,7 +16,6 @@
 #include "Selection/MeshSelectionMechanic.h"
 #include "Selection/DynamicMeshSelection.h"
 #include "ToolSetupUtil.h"
-#include "UVToolContextObjects.h"
 
 #include "ToolTargetManager.h"
 
@@ -23,6 +23,216 @@
 
 using namespace UE::Geometry;
 
+namespace UVSelectToolLocals
+{
+	/**
+	 * An undo/redo object for selection changes that, instead of operating directly on a selection
+	 * mechanic, instead operates on a context object that tools can use to route the request
+	 * to the current selection mechanic. This is valuable because we want the selection changes
+	 * to be undoable in different invocations of the tool, and the selection mechanic pointer
+	 * will not stay the same. However, the context object will stay the same, and we can register
+	 * to its delegate on each invocation.
+	 */
+	class FSelectionChange : public FToolCommandChange
+	{
+	public:
+		/**
+		 * @param bBroadcastOnSelectionChangedIn Whether the change in selection should broadcast
+		 *   OnSelectionChanged, which updates gizmo, etc.
+		 * @param GizmoBeforeIn Only relevant if bBroadcastOnSelectionChangedIn is true. In that case,
+		 *   the gizmo gets reset on the way forward to the current selection, which means we have to
+		 *   reset it to the old orientation on the way back (otherwise a rotated gizmo would end up
+		 *   losing its rotation on undo).
+		 */
+		FSelectionChange(const FDynamicMeshSelection& SelectionBeforeIn,
+			const FDynamicMeshSelection& SelectionAfterIn,
+			bool bBroadcastOnSelectionChangedIn,
+			const FTransform& GizmoBeforeIn)
+			: SelectionBefore(SelectionBeforeIn)
+			, SelectionAfter(SelectionAfterIn)
+			, bBroadcastOnSelectionChanged(bBroadcastOnSelectionChangedIn)
+			, GizmoBefore(GizmoBeforeIn)
+		{
+		}
+
+		virtual void Apply(UObject* Object) override
+		{
+			UUVSelectToolChangeRouter* ChangeRouter = Cast<UUVSelectToolChangeRouter>(Object);
+			if (ensure(ChangeRouter) && ChangeRouter->CurrentSelectTool.IsValid())
+			{
+				ChangeRouter->CurrentSelectTool->SetSelection(SelectionAfter, bBroadcastOnSelectionChanged);
+			}
+		}
+
+		virtual void Revert(UObject* Object) override
+		{
+			UUVSelectToolChangeRouter* ChangeRouter = Cast<UUVSelectToolChangeRouter>(Object);
+			if (ensure(ChangeRouter) && ChangeRouter->CurrentSelectTool.IsValid())
+			{
+				ChangeRouter->CurrentSelectTool->SetSelection(SelectionBefore, bBroadcastOnSelectionChanged);
+				if (bBroadcastOnSelectionChanged)
+				{
+					ChangeRouter->CurrentSelectTool->SetGizmoTransform(GizmoBefore);
+				}
+			}
+		}
+
+		virtual FString ToString() const override
+		{
+			return TEXT("UVSelectToolLocals::FSelectionChange");
+		}
+
+	protected:
+		FDynamicMeshSelection SelectionBefore;
+		FDynamicMeshSelection SelectionAfter;
+		bool bBroadcastOnSelectionChanged;
+		FTransform GizmoBefore;
+	};
+
+	/**
+	 * A change that allows another change to be injected into it even after it is placed in the undo
+	 * stack.
+	 * NOTE: Remove this if we merge Transform into Select tool and choose not to store/restore
+	 * selecitons in the Select tool.
+	 */
+	class FSpeculativeSelectionChange : public FToolCommandChange
+	{
+	public:
+		FSpeculativeSelectionChange(TSharedPtr<TUniquePtr<FToolCommandChange>> PointerToContentIn)
+			: PointerToContent(PointerToContentIn)
+		{
+		}
+
+		virtual bool HasExpired(UObject* Object) const override
+		{
+			return !PointerToContent->IsValid();
+		}
+
+		virtual void Apply(UObject* Object) override
+		{
+			if (*PointerToContent)
+			{
+				(*PointerToContent)->Apply(Object);
+			}
+		}
+
+		virtual void Revert(UObject* Object) override
+		{
+			if (*PointerToContent)
+			{
+				(*PointerToContent)->Revert(Object);
+			}
+		}
+
+		virtual FString ToString() const override
+		{
+			return TEXT("UVSelectToolLocals::FSpeculativeSelectionChange");
+		}
+
+		TSharedPtr<TUniquePtr<FToolCommandChange>> PointerToContent;
+	};
+
+	/**
+	 * A change similar to the one emitted by EmitChangeApi->EmitToolIndependentUnwrapCanonicalChange,
+	 * but which updates the Select tool's gizmo in a way that preserves the rotational component
+	 * (which would be lost if we just updated the gizmo from the current selection on undo/redo).
+	 * 
+	 * There is some built-in change tracking for the gizmo component in our transform gizmo, but 
+	 * due to the order in which changes get emitted, there is not a good way to make sure that we
+	 * update the selection mechanic (which needs to know the gizmo transform) at the correct time
+	 * relative to those built-in changes. So, those built-in changes are actually wasted on us,
+	 * but it was not easy to deactivate them because the change emitter is linked to the transform
+	 * proxy...
+	 *
+	 * Expects UUVSelectToolChangeRouter to be the passed-in object
+	 */
+	class  FGizmoMeshChange : public FToolCommandChange
+	{
+	public:
+		FGizmoMeshChange(UUVEditorToolMeshInput* UVToolInputObjectIn,
+			TUniquePtr<UE::Geometry::FDynamicMeshChange> UnwrapCanonicalMeshChangeIn,
+			const FTransform& GizmoBeforeIn, const FTransform& GizmoAfterIn)
+
+			: UVToolInputObject(UVToolInputObjectIn)
+			, UnwrapCanonicalMeshChange(MoveTemp(UnwrapCanonicalMeshChangeIn))
+			, GizmoBefore(GizmoBeforeIn)
+			, GizmoAfter(GizmoAfterIn)
+		{
+			ensure(UVToolInputObjectIn);
+			ensure(UnwrapCanonicalMeshChange);
+		};
+
+		virtual void Apply(UObject* Object) override
+		{
+			UnwrapCanonicalMeshChange->Apply(UVToolInputObject->UnwrapCanonical.Get(), false);
+			UVToolInputObject->UpdateFromCanonicalUnwrapUsingMeshChange(*UnwrapCanonicalMeshChange);
+			
+			// This is a little wasteful because we're going to reset the gizmo transform, 
+			// but it updates the AABBTree for us.
+			UVToolInputObject->OnUndoRedo.Broadcast(false);
+
+			UUVSelectToolChangeRouter* ChangeRouter = Cast<UUVSelectToolChangeRouter>(Object);
+			if (ensure(ChangeRouter) && ChangeRouter->CurrentSelectTool.IsValid())
+			{
+				ChangeRouter->CurrentSelectTool->SetGizmoTransform(GizmoAfter);
+			}
+
+		}
+
+		virtual void Revert(UObject* Object) override
+		{
+			UnwrapCanonicalMeshChange->Apply(UVToolInputObject->UnwrapCanonical.Get(), true);
+			UVToolInputObject->UpdateFromCanonicalUnwrapUsingMeshChange(*UnwrapCanonicalMeshChange);
+			
+			// This is a little wasteful because we're going to reset the gizmo transform, 
+			// but it updates the AABBTree for us.
+			UVToolInputObject->OnUndoRedo.Broadcast(false);
+
+			UUVSelectToolChangeRouter* ChangeRouter = Cast<UUVSelectToolChangeRouter>(Object);
+			if (ensure(ChangeRouter) && ChangeRouter->CurrentSelectTool.IsValid())
+			{
+				ChangeRouter->CurrentSelectTool->SetGizmoTransform(GizmoBefore);
+			}
+		}
+
+		virtual bool HasExpired(UObject* Object) const override
+		{
+			return !(UVToolInputObject.IsValid() && UnwrapCanonicalMeshChange);
+		}
+
+
+		virtual FString ToString() const override
+		{
+			return TEXT("UVSelectToolLocals::FGizmoMeshChange");
+		}
+
+	protected:
+		TWeakObjectPtr<UUVEditorToolMeshInput> UVToolInputObject;
+		TUniquePtr<UE::Geometry::FDynamicMeshChange> UnwrapCanonicalMeshChange;
+		FTransform GizmoBefore;
+		FTransform GizmoAfter;
+	};
+}
+
+void UUVSelectToolSpeculativeChangeAPI::EmitSpeculativeChange(UObject* TargetObject, 
+	UUVToolEmitChangeAPI* EmitChangeAPI, const FText& TransactionName)
+{
+	ContentOfLastSpeculativeChange = MakeShared<TUniquePtr<FToolCommandChange>>();
+
+	EmitChangeAPI->EmitToolIndependentChange(TargetObject, 
+		MakeUnique<UVSelectToolLocals::FSpeculativeSelectionChange>(ContentOfLastSpeculativeChange), 
+		TransactionName);
+}
+
+void UUVSelectToolSpeculativeChangeAPI::InsertIntoLastSpeculativeChange(TUniquePtr<FToolCommandChange> ChangeToInsert)
+{
+	if (ContentOfLastSpeculativeChange)
+	{
+		*ContentOfLastSpeculativeChange = MoveTemp(ChangeToInsert);
+	}
+}
+
+
 /*
  * ToolBuilder
  */
@@ -60,15 +270,35 @@ void UUVSelectTool::Setup()
 	Settings->RestoreProperties(this);
 	AddToolPropertySource(Settings);
 
+	UContextObjectStore* ContextStore = GetToolManager()->GetContextObjectStore();
+	EmitChangeAPI = ContextStore->FindContext<UUVToolEmitChangeAPI>();
+
 	SelectionMechanic = NewObject<UMeshSelectionMechanic>();
 	SelectionMechanic->Setup(this);
 	SelectionMechanic->SetWorld(Targets[0]->UnwrapPreview->GetWorld());
 	SelectionMechanic->OnSelectionChanged.AddUObject(this, &UUVSelectTool::OnSelectionChanged);
 
+	// Make it so that our selection mechanic creates undo/redo transactions that go to a selection
+	// change router, which we use to route to the current selection mechanic on each tool invocation.
+	ChangeRouter = ContextStore->FindContext<UUVSelectToolChangeRouter>();
+	if (!ChangeRouter)
+	{
+		ChangeRouter = NewObject<UUVSelectToolChangeRouter>();
+		ContextStore->AddContextObject(ChangeRouter);
+	}
+	ChangeRouter->CurrentSelectTool = this;
+
+	SelectionMechanic->EmitSelectionChange = [this](const FDynamicMeshSelection& OldSelection,
+		const FDynamicMeshSelection& NewSelection, bool bBroadcastOnSelectionChangedIn)
+	{
+		EmitChangeAPI->EmitToolIndependentChange(ChangeRouter, MakeUnique<UVSelectToolLocals::FSelectionChange>(
+			OldSelection, NewSelection, bBroadcastOnSelectionChangedIn, TransformGizmo->GetGizmoTransform()), 
+			LOCTEXT("SelectionChangeMessage", "Selection Change"));
+	};
+
 	ConfigureSelectionModeFromControls();
 
 	// Retrieve cached AABB tree storage, or else set it up
-	UContextObjectStore* ContextStore = GetToolManager()->GetContextObjectStore();
 	UUVToolAABBTreeStorage* TreeStore = ContextStore->FindContext<UUVToolAABBTreeStorage>();
 	if (!TreeStore)
 	{
@@ -106,6 +336,37 @@ void UUVSelectTool::Setup()
 	{
 		SelectionMechanic->SetSelection(*SelectionStore->Selection);
 	}
+	if (!SelectionMechanic->GetCurrentSelection().IsEmpty() && !SelectionMechanic->GetCurrentSelection().MatchesTimestamp())
+	{
+		// If we have an expired selection, then there must have been some tool that changed the
+		// topology. We need to clear the selection in an undoable way, but that clear actually
+		// needs to happen before the last invocation shutdown to be in the proper sequence in the 
+		// undo stack.
+		FDynamicMeshSelection NewSelection;
+		UUVSelectToolSpeculativeChangeAPI* SpeculativeChangeAPI = ContextStore->FindContext<UUVSelectToolSpeculativeChangeAPI>();
+		if (SpeculativeChangeAPI && SpeculativeChangeAPI->HasSpeculativeChange())
+		{
+			SpeculativeChangeAPI->InsertIntoLastSpeculativeChange(MakeUnique<UVSelectToolLocals::FSelectionChange>(
+				SelectionMechanic->GetCurrentSelection(), NewSelection, false, FTransform()));
+		}
+
+		SelectionMechanic->SetSelection(NewSelection, false, false);
+	}
+
+	// Make sure that if we receive undo/redo events on the meshes, we update the tree structures
+	// and the selection mechanic drawn elements. Note that we mainly have to worry about this
+	// because the select tool is the default UV editor tool, and as such it can receive undo
+	// transactions from other tools and from other select tool invocations. Other tools typically
+	// only need to worry about their own transactions, since we undo other tool invocations before
+	// we get to unrelated transactions, and we can't redo out of the default tool.
+	for (int32 i = 0; i < Targets.Num(); ++i)
+	{
+		Targets[i]->OnUndoRedo.AddWeakLambda(this, [this, i](bool bRevert) {
+			AABBTrees[i]->Build();
+			UpdateGizmo();
+			SelectionMechanic->RebuildDrawnElements(TransformGizmo->GetGizmoTransform());
+			});
+	}
 
 	// Gizmo setup
 	UInteractiveGizmoManager* GizmoManager = GetToolManager()->GetPairedGizmoManager();
@@ -143,6 +404,8 @@ void UUVSelectTool::Shutdown(EToolShutdownType ShutdownType)
 	UContextObjectStore* ContextStore = GetToolManager()->GetContextObjectStore();
 	if (!SelectionMechanic->GetCurrentSelection().IsEmpty())
 	{
+		// TODO: It's not clear whether storing our selection is the right choice in the long run,
+		// though we currently need it while the gizmo transformation is a separate tool... 
 		UUVToolMeshSelection* SelectionStore = ContextStore->FindContext<UUVToolMeshSelection>();
 		if (!SelectionStore)
 		{
@@ -150,12 +413,40 @@ void UUVSelectTool::Shutdown(EToolShutdownType ShutdownType)
 			ContextStore->AddContextObject(SelectionStore);
 		}
 		*SelectionStore->Selection = SelectionMechanic->GetCurrentSelection();
+
+		// We update the topology timestamp here because the timestamps may have divereged due
+		// to our undo/redo change implementation sometimes incrementing the topology timestamp
+		// even when the topology hasn't changed.
+		// TODO: Check whether we still need this once we use FMeshVertexChange instead of FDynamicMeshChange
+		SelectionStore->Selection->TopologyTimestamp = SelectionStore->Selection->Mesh->GetTopologyTimestamp();
+
+		// Don't issue a transaciton if we're cancelling (currently only possible via undo out
+		// of the "Transform" tool). Otherwise, issue a speculative change so that we can clear
+		// the selection if the mesh topology turns out to have changed by the time we run again.
+		if (ShutdownType != EToolShutdownType::Cancel)
+		{
+			UUVSelectToolSpeculativeChangeAPI* SpeculativeChangeAPI = ContextStore->FindContext<UUVSelectToolSpeculativeChangeAPI>();
+			if (!SpeculativeChangeAPI)
+			{
+				SpeculativeChangeAPI = NewObject<UUVSelectToolSpeculativeChangeAPI>();
+				ContextStore->AddContextObject(SpeculativeChangeAPI);
+			}
+			SpeculativeChangeAPI->EmitSpeculativeChange(ChangeRouter, EmitChangeAPI,
+				LOCTEXT("SpeculativeSlectionChangeName", "End Select Tool"));
+		}
 	}
 	else
 	{
 		ContextStore->RemoveContextObjectsOfType<UUVToolMeshSelection>();
 	}
 
+	ChangeRouter->CurrentSelectTool = nullptr;
+
+	for (TObjectPtr<UUVEditorToolMeshInput> Target : Targets)
+	{
+		Target->OnUndoRedo.RemoveAll(this);
+	}
+
 	Settings->SaveProperties(this);
 
 	SelectionMechanic->Shutdown();
@@ -168,6 +459,20 @@ void UUVSelectTool::Shutdown(EToolShutdownType ShutdownType)
 
 	// Calls shutdown on gizmo and destroys it.
 	GetToolManager()->GetPairedGizmoManager()->DestroyAllGizmosByOwner(this);
+
+	EmitChangeAPI = nullptr;
+	ChangeRouter = nullptr;
+}
+
+void UUVSelectTool::SetSelection(const UE::Geometry::FDynamicMeshSelection& NewSelection, bool bBroadcastOnSelectionChanged)
+{
+	SelectionMechanic->SetSelection(NewSelection, bBroadcastOnSelectionChanged, false);
+}
+
+void UUVSelectTool::SetGizmoTransform(const FTransform& NewTransform)
+{
+	TransformGizmo->ReinitializeGizmoTransform(NewTransform);
+	SelectionMechanic->RebuildDrawnElements(NewTransform);
 }
 
 void UUVSelectTool::OnPropertyModified(UObject* PropertySet, FProperty* Property)
@@ -183,11 +488,6 @@ void UUVSelectTool::UpdateGizmo()
 	{
 		FVector3d Centroid = SelectionMechanic->GetCurrentSelectionCentroid();
 
-		// TODO: We should either make it so that transparent materials don't affect the rendering
-		// of the gizmo component, or else we should pass in some depth offset into the component.
-		// For now we hack it by raising the gizmo a little off the plane manually.
-		Centroid += FVector3d::ZAxisVector * KINDA_SMALL_NUMBER;
-
 		TransformGizmo->ReinitializeGizmoTransform(FTransform((FVector)Centroid));
 	}
 
@@ -377,6 +677,8 @@ void UUVSelectTool::SetGizmoEnabled(bool bEnabledIn)
 
 void UUVSelectTool::GizmoTransformStarted(UTransformProxy* Proxy)
 {
+	bInDrag = true;
+
 	InitialGizmoFrame = FFrame3d(TransformGizmo->ActiveTarget->GetTransform());
 	MovingVertOriginalPositions.SetNum(MovingVids.Num());
 	const FDynamicMesh3* Mesh = Targets[SelectionTargetIndex]->UnwrapCanonical.Get();
@@ -390,6 +692,14 @@ void UUVSelectTool::GizmoTransformStarted(UTransformProxy* Proxy)
 
 void UUVSelectTool::GizmoTransformChanged(UTransformProxy* Proxy, FTransform Transform)
 {
+	// This function gets called both during drag and on undo/redo. This might have been ok if
+	// undo/redo also called GizmoTransformStarted/GizmoTransformEnded, but they don't, which
+	// means the two types of events operate quite differently. We just ignore any non-drag calls.
+	if (!bInDrag)
+	{
+		return;
+	}
+
 	FTransform DeltaTransform = Transform.GetRelativeTransform(InitialGizmoFrame.ToFTransform());
 
 	if (!DeltaTransform.GetTranslation().IsNearlyZero() || !DeltaTransform.GetRotation().IsIdentity() || Transform.GetScale3D() != FVector::One())
@@ -401,10 +711,16 @@ void UUVSelectTool::GizmoTransformChanged(UTransformProxy* Proxy, FTransform Tra
 
 void UUVSelectTool::GizmoTransformEnded(UTransformProxy* Proxy)
 {
-	// TODO: Add undo support
-	//FDynamicMeshChangeTracker ChangeTracker(Targets[SelectionTargetIndex]->UnwrapPreview->PreviewMesh->GetMesh());
-	//ChangeTracker.BeginChange();
-	//ChangeTracker.SaveTriangles(MovingVids, true);
+	bInDrag = false;
+
+	// Set things up for undo.
+	// TODO: We should really use FMeshVertexChange instead of FDynamicMeshChange because we don't
+	// need to alter the mesh topology. However we currently don't have a way to apply a FMeshVertexChange
+	// directly to a dynamic mesh pointer, only via UDynamicMesh. We should change things here once
+	// that ability exists.
+	FDynamicMeshChangeTracker ChangeTracker(Targets[SelectionTargetIndex]->UnwrapCanonical.Get());
+	ChangeTracker.BeginChange();
+	ChangeTracker.SaveTriangles(SelectedTids, true);
 
 	// One final attempt to apply transforms if OnTick hasn't happened yet
 	ApplyGizmoTransform();
@@ -424,6 +740,14 @@ void UUVSelectTool::GizmoTransformEnded(UTransformProxy* Proxy)
 		AABBTrees[SelectionTargetIndex]->Build();
 	}
 
+	const FText TransactionName(LOCTEXT("DragCompleteTransactionName", "Move Items"));
+	EmitChangeAPI->BeginUndoTransaction(TransactionName);
+	EmitChangeAPI->EmitToolIndependentChange(ChangeRouter, MakeUnique<UVSelectToolLocals::FGizmoMeshChange>(
+		Targets[SelectionTargetIndex], ChangeTracker.EndChange(), 
+		InitialGizmoFrame.ToFTransform(), TransformGizmo->GetGizmoTransform()),
+		TransactionName);
+	EmitChangeAPI->EndUndoTransaction();
+
 	TransformGizmo->SetNewChildScale(FVector::One());
 	SelectionMechanic->RebuildDrawnElements(TransformGizmo->GetGizmoTransform());
 }
diff --git a/Engine/Plugins/Experimental/UVEditor/Source/UVEditorTools/Private/UVToolContextObjects.cpp b/Engine/Plugins/Experimental/UVEditor/Source/UVEditorTools/Private/UVToolContextObjects.cpp
--- a/Engine/Plugins/Experimental/UVEditor/Source/UVEditorTools/Private/UVToolContextObjects.cpp
+++ b/Engine/Plugins/Experimental/UVEditor/Source/UVEditorTools/Private/UVToolContextObjects.cpp
@@ -2,11 +2,81 @@
 
 #include "UVToolContextObjects.h"
 
-#include "InputRouter.h" // Need to define this and UWorld so weak pointers know they are UThings
+#include "DynamicMesh/MeshIndexUtil.h"
 #include "Engine/World.h"
+#include "InputRouter.h" // Need to define this and UWorld so weak pointers know they are UThings
+#include "InteractiveToolManager.h"
+#include "ToolTargets/UVEditorToolMeshInput.h"
 
 using namespace UE::Geometry;
 
+namespace UVToolContextObjectLocals
+{
+
+	/**
+	 * A wrapper change that applies a given change to the unwrap canonical mesh of an input, uses that to update the
+	 * other views, and issues an OnUndoRedo broadcast.
+	 */
+	class  FUVEditorMeshChange : public FToolCommandChange
+	{
+	public:
+		FUVEditorMeshChange(UUVEditorToolMeshInput* UVToolInputObjectIn, TUniquePtr<FDynamicMeshChange> UnwrapCanonicalMeshChangeIn)
+			: UVToolInputObject(UVToolInputObjectIn)
+			, UnwrapCanonicalMeshChange(MoveTemp(UnwrapCanonicalMeshChangeIn))
+		{
+			ensure(UVToolInputObjectIn);
+			ensure(UnwrapCanonicalMeshChange);
+		};
+
+		virtual void Apply(UObject* Object) override
+		{
+			UnwrapCanonicalMeshChange->Apply(UVToolInputObject->UnwrapCanonical.Get(), false);
+			UVToolInputObject->UpdateFromCanonicalUnwrapUsingMeshChange(*UnwrapCanonicalMeshChange);
+			UVToolInputObject->OnUndoRedo.Broadcast(false);
+		}
+
+		virtual void Revert(UObject* Object) override
+		{
+			UnwrapCanonicalMeshChange->Apply(UVToolInputObject->UnwrapCanonical.Get(), true);
+			UVToolInputObject->UpdateFromCanonicalUnwrapUsingMeshChange(*UnwrapCanonicalMeshChange);
+			UVToolInputObject->OnUndoRedo.Broadcast(true);
+		}
+
+		virtual bool HasExpired(UObject* Object) const override
+		{
+			return !(UVToolInputObject.IsValid() && UnwrapCanonicalMeshChange);
+		}
+
+
+		virtual FString ToString() const override
+		{
+			return TEXT("FUVEditorMeshChange");
+		}
+
+	protected:
+		TWeakObjectPtr<UUVEditorToolMeshInput> UVToolInputObject;
+		TUniquePtr<FDynamicMeshChange> UnwrapCanonicalMeshChange;
+	};
+}
+
+void UUVToolEmitChangeAPI::EmitToolIndependentChange(UObject* TargetObject, TUniquePtr<FToolCommandChange> Change, const FText& Description)
+{
+	ToolManager->GetContextTransactionsAPI()->AppendChange(TargetObject, MoveTemp(Change), Description);
+}
+
+void UUVToolEmitChangeAPI::EmitToolIndependentUnwrapCanonicalChange(UUVEditorToolMeshInput* InputObject, 
+	TUniquePtr<FDynamicMeshChange> UnwrapCanonicalMeshChange, const FText& Description)
+{
+	ToolManager->GetContextTransactionsAPI()->AppendChange(InputObject, 
+		MakeUnique<UVToolContextObjectLocals::FUVEditorMeshChange>(InputObject, MoveTemp(UnwrapCanonicalMeshChange)), Description);
+}
+
+void UUVToolEmitChangeAPI::EmitToolDependentChange(UObject* TargetObject, TUniquePtr<FToolCommandChange> Change, const FText& Description)
+{
+	// This should wrap the change in the proper wrapper that will expire it when the tool changes
+	ToolManager->EmitObjectChange(TargetObject, MoveTemp(Change), Description);
+}
+
 void UUVToolLivePreviewAPI::Initialize(UWorld* WorldIn, UInputRouter* RouterIn)
 {
 	World = WorldIn;
diff --git a/Engine/Plugins/Experimental/UVEditor/Source/UVEditorTools/Public/ToolTargets/UVEditorToolMeshInput.h b/Engine/Plugins/Experimental/UVEditor/Source/UVEditorTools/Public/ToolTargets/UVEditorToolMeshInput.h
--- a/Engine/Plugins/Experimental/UVEditor/Source/UVEditorTools/Public/ToolTargets/UVEditorToolMeshInput.h
+++ b/Engine/Plugins/Experimental/UVEditor/Source/UVEditorTools/Public/ToolTargets/UVEditorToolMeshInput.h
@@ -16,6 +16,7 @@ class UMaterialInterface;
 class UMeshOpPreviewWithBackgroundCompute;
 class UMeshElementsVisualizer;
 PREDECLARE_GEOMETRY(class FDynamicMesh3);
+PREDECLARE_GEOMETRY(class FDynamicMeshChange);
 
 /**
  * A package of the needed information for an asset being operated on by a
@@ -78,6 +79,15 @@ public:
 	UPROPERTY()
 	TObjectPtr<UMeshElementsVisualizer> WireframeDisplay = nullptr;
 
+	/**
+	 * Helper delegate that allows undo/redo transactions to notify the active tool that the input
+	 * object was modified. This is helpful for allowing transactions performed in one tool invocation
+	 * to properly be reverted/applied during a different (default) tool invocation, since the new
+	 * tool invocation can listen to this delegate and update relevant structures on broadcast.
+	 */
+	DECLARE_MULTICAST_DELEGATE_OneParam(FOnUndoRedo, bool bRevert);
+	FOnUndoRedo OnUndoRedo;
+
 	// Additional needed information
 	TObjectPtr<UObject> OriginalAsset = nullptr;
 	int32 UVLayerIndex = 0;
@@ -162,6 +172,12 @@ public:
 	 */
 	void UpdateAllFromAppliedPreview(const TArray<int32>* ChangedElementIDs = nullptr, const TArray<int32>* ChangedConnectivityTids = nullptr, const TArray<int32>* FastRenderUpdateTids = nullptr);
 
+	/**
+	 * Uses the stored triangles/vertices in the mesh change to update everything from the canonical unwrap. 
+	 * Assumes the change has already been applied to the canonical unwrap.
+	 */
+	void UpdateFromCanonicalUnwrapUsingMeshChange(const UE::Geometry::FDynamicMeshChange& UnwrapCanonicalMeshChange);
+
 	// UToolTarget
 	virtual bool IsValid() const override;
 protected:
diff --git a/Engine/Plugins/Experimental/UVEditor/Source/UVEditorTools/Public/UVSelectTool.h b/Engine/Plugins/Experimental/UVEditor/Source/UVEditorTools/Public/UVSelectTool.h
--- a/Engine/Plugins/Experimental/UVEditor/Source/UVEditorTools/Public/UVSelectTool.h
+++ b/Engine/Plugins/Experimental/UVEditor/Source/UVEditorTools/Public/UVSelectTool.h
@@ -6,23 +6,28 @@
 
 #include "DynamicMesh/DynamicMeshAABBTree3.h"
 #include "FrameTypes.h"
+#include "GeometryBase.h"
 #include "InteractiveTool.h"
-//#include "InteractiveToolChange.h"
 #include "InteractiveToolBuilder.h"
 #include "TargetInterfaces/UVUnwrapDynamicMesh.h"
+#include "UVToolContextObjects.h"
 
 #include "UVSelectTool.generated.h"
 
+PREDECLARE_GEOMETRY(class FDynamicMeshSelection);
 class APreviewGeometryActor;
+class FToolCommandChange;
 class ULineSetComponent;
 class UMeshSelectionMechanic;
 class UToolTargetManager;
 class UTransformGizmo;
 class UTransformProxy;
 class UUVEditorMode;
+class UUVSelectToolChangeRouter;
 class UUVToolStateObjectStore;
 class UPreviewMesh;
 class UUVEditorToolMeshInput;
+class UUVToolEmitChangeAPI;
 
 UCLASS()
 class UVEDITORTOOLS_API UUVSelectToolBuilder : public UInteractiveToolBuilder
@@ -93,6 +98,10 @@ public:
 		Targets = TargetsIn;
 	}
 
+	// Used by undo/redo changes to update the tool state
+	void SetSelection(const UE::Geometry::FDynamicMeshSelection& NewSelection, bool bBroadcastOnSelectionChanged);
+	void SetGizmoTransform(const FTransform& NewTransform);
+
 	// UInteractiveTool
 	virtual void Setup() override;
 	virtual void Shutdown(EToolShutdownType ShutdownType) override;
@@ -129,18 +138,25 @@ protected:
 	TObjectPtr<UMeshSelectionMechanic> SelectionMechanic = nullptr;
 
 	UPROPERTY()
-	UTransformGizmo* TransformGizmo = nullptr;
+	TObjectPtr<UTransformGizmo> TransformGizmo = nullptr;
 
 	TArray<TSharedPtr<FDynamicMeshAABBTree3>> AABBTrees;
 
 	UPROPERTY()
-	APreviewGeometryActor* LivePreviewGeometryActor = nullptr;
+	TObjectPtr<APreviewGeometryActor> LivePreviewGeometryActor = nullptr;
+
+	UPROPERTY()
+	TObjectPtr<ULineSetComponent> LivePreviewLineSet = nullptr;
+
+	UPROPERTY()
+	TObjectPtr<UUVToolEmitChangeAPI> EmitChangeAPI = nullptr;
 
 	UPROPERTY()
-	ULineSetComponent* LivePreviewLineSet = nullptr;
+	TObjectPtr<UUVSelectToolChangeRouter> ChangeRouter = nullptr;
 
 	UE::Geometry::FFrame3d InitialGizmoFrame;
 	FTransform UnappliedGizmoTransform;
+	bool bInDrag = false;
 	bool bGizmoTransformNeedsApplication = false;
 
 	TArray<int32> MovingVids;
@@ -153,27 +169,51 @@ protected:
 	bool bGizmoEnabled = false;
 };
 
-///**
-// * Wraps a FDynamicMeshChange.
-// */
-//class UVEDITORTOOLS_API FUVSelectToolChange : public FToolCommandChange
-//{
-//public:
-//	FUVSelectToolChange(TUniquePtr<FDynamicMeshChange> MeshChangeIn)
-//		: MeshChange(MoveTemp(MeshChangeIn))
-//	{};
-//
-//	virtual void Apply(UObject* Object) override;
-//	virtual void Revert(UObject* Object) override;
-//	//virtual bool HasExpired(UObject* Object) const override
-//	//{
-//	//	return Cast<UGroupEdgeInsertionTool>(Object)->CurrentChangeStamp != ChangeStamp;
-//	//}
-//	virtual FString ToString() const override
-//	{
-//		return TEXT("FUVSelectToolChange");
-//	}
-//
-//protected:
-//	TUniquePtr<FDynamicMeshChange> MeshChange;
-//};
+/**
+ * A helper context object that we can use as the target of undo/redo events to apply them
+ * to the current invocation of the select tool (which may have different gizmo/selection
+ * pointers than those that were around when the change was emitted).
+ */
+UCLASS()
+class UVEDITORTOOLS_API UUVSelectToolChangeRouter : public UUVToolContextObject
+{
+	GENERATED_BODY()
+public:
+
+	TWeakObjectPtr<UUVSelectTool> CurrentSelectTool = nullptr;
+};
+
+/**
+ * A helper context object that allows us to inject an undo transaction back in time, which gets
+ * used to deal with the fact that our stored selection may become invalidated by an intervening
+ * tool, and needs to be cleared in an undoable transaction before that tool runs.
+ * 
+ * NOTE: It seems likely (especially after "Transform" is no longer a separate tool from "Select")
+ * that we may just decide not to try to store a selection on tool shutdown. In that case we will
+ * delete this class.
+ */
+UCLASS()
+class UVEDITORTOOLS_API UUVSelectToolSpeculativeChangeAPI : public UUVToolContextObject
+{
+	GENERATED_BODY()
+public:
+
+	/** 
+	 * Emits a tool-independent change that does nothing unless a subsequent InsertIntoLastSpeculativeChange
+	 * call injects a change.
+	 */
+	void EmitSpeculativeChange(UObject* TargetObject, UUVToolEmitChangeAPI* EmitChangeAPI, const FText& TransactionName);
+
+	bool HasSpeculativeChange()
+	{
+		return ContentOfLastSpeculativeChange.IsValid();
+	}
+
+	/**
+	 * Inserts a change into the place marked by the last EmitSpeculativeChange call.
+	 */
+	void InsertIntoLastSpeculativeChange(TUniquePtr<FToolCommandChange> ChangeToInsert);
+
+protected:
+	TSharedPtr<TUniquePtr<FToolCommandChange>> ContentOfLastSpeculativeChange;
+};
\ No newline at end of file
diff --git a/Engine/Plugins/Experimental/UVEditor/Source/UVEditorTools/Public/UVToolContextObjects.h b/Engine/Plugins/Experimental/UVEditor/Source/UVEditorTools/Public/UVToolContextObjects.h
--- a/Engine/Plugins/Experimental/UVEditor/Source/UVEditorTools/Public/UVToolContextObjects.h
+++ b/Engine/Plugins/Experimental/UVEditor/Source/UVEditorTools/Public/UVToolContextObjects.h
@@ -5,6 +5,8 @@
 #include "CoreMinimal.h"
 
 #include "DynamicMesh/DynamicMeshAABBTree3.h"
+#include "DynamicMesh/DynamicMeshChangeTracker.h" // FDynamicMeshChange for TUniquePtr
+#include "InteractiveToolManager.h"
 #include "Selection/DynamicMeshSelection.h"
 #include "GeometryBase.h"
 
@@ -13,8 +15,10 @@
 // TODO: Should this be spread out across multiple files?
 
 PREDECLARE_GEOMETRY(class FDynamicMesh3);
+class FToolCommandChange;
 class UInputRouter;
 class UWorld;
+class UUVEditorToolMeshInput;
 
 UCLASS()
 class UVEDITORTOOLS_API UUVToolContextObject : public UObject
@@ -23,6 +27,61 @@ class UVEDITORTOOLS_API UUVToolContextObject : public UObject
 public:
 };
 
+/**
+ * An API object meant to be stored in a context object store that allows UV editor tools
+ * to emit appropriate undo/redo transactions.
+ */
+UCLASS()
+class UVEDITORTOOLS_API UUVToolEmitChangeAPI : public UUVToolContextObject
+{
+	GENERATED_BODY()
+public:
+
+	void Initialize(TObjectPtr<UInteractiveToolManager> ToolManagerIn)
+	{
+		ToolManager = ToolManagerIn;
+	}
+
+	virtual void BeginUndoTransaction(const FText& Description)
+	{
+		ToolManager->BeginUndoTransaction(Description);
+	}
+	virtual void EndUndoTransaction()
+	{
+		ToolManager->EndUndoTransaction();
+	}
+
+	/**
+	 * Emit a change that can be undone even if we leave the tool from which it is emitted (as
+	 * long as that UV editor instance is still open). 
+	 * Minor note: because we undo "out of" tools into a default tool and never out of a default tool,
+	 * in practice, tool-independent changes will only ever be applied/reverted in the same tool 
+	 * invocation that they were emitted or in the default tool, not in other arbitrary tools.
+	 * 
+	 * Since tool-independent changes usually operate on UV editor mesh input object, it is probably
+	 * preferable to use EmitToolIndependentUnwrapCanonicalChange, which will set up a proper transaction
+	 * for you.
+	 */
+	virtual void EmitToolIndependentChange(UObject* TargetObject, TUniquePtr<FToolCommandChange> Change, const FText& Description);
+
+	/**
+	 * A convenience function that is like EmitToolIndependentChange, but uses a FDynamicMeshChange
+	 * that operates on the UnwrapCanonical of an input to create a change object that updates the other
+	 * views and issues an OnUndoRedo broadcast on the input object.
+	 */
+	virtual void EmitToolIndependentUnwrapCanonicalChange(UUVEditorToolMeshInput* InputObject,
+		TUniquePtr<UE::Geometry::FDynamicMeshChange> UnwrapCanonicalMeshChange, const FText& Description);
+
+	/**
+	 * Emits a change that is considered expired when the active tool does not match the tool that was active
+	 * when it was emitted.
+	 */
+	virtual void EmitToolDependentChange(UObject* TargetObject, TUniquePtr<FToolCommandChange> Change, const FText& Description);
+
+protected:
+	TObjectPtr<UInteractiveToolManager> ToolManager = nullptr;
+};
+
 UCLASS()
 class UVEDITORTOOLS_API UUVToolLivePreviewAPI : public UUVToolContextObject
 {
