GeometryProcessing: remove T* conversions/casts from FVector2<T>
#rnx

diff --git a/Engine/Plugins/Experimental/UVEditor/Source/UVEditorTools/Private/ToolTargets/UVEditorToolMeshInput.cpp b/Engine/Plugins/Experimental/UVEditor/Source/UVEditorTools/Private/ToolTargets/UVEditorToolMeshInput.cpp
--- a/Engine/Plugins/Experimental/UVEditor/Source/UVEditorTools/Private/ToolTargets/UVEditorToolMeshInput.cpp
+++ b/Engine/Plugins/Experimental/UVEditor/Source/UVEditorTools/Private/ToolTargets/UVEditorToolMeshInput.cpp
@@ -76,7 +76,8 @@ namespace UVEditorToolMeshInputLocals
 			}
 			else
 			{
-				OverlayOut.InsertElement(ElementID, (float*)OverlayIn.GetElement(ElementID));
+				FVector2f UV = OverlayIn.GetElement(ElementID);
+				OverlayOut.InsertElement(ElementID, &UV.X);
 			}
 		};
 
diff --git a/Engine/Plugins/Experimental/UVEditor/Source/UVEditorTools/Private/UVEditorToolUtil.cpp b/Engine/Plugins/Experimental/UVEditor/Source/UVEditorTools/Private/UVEditorToolUtil.cpp
--- a/Engine/Plugins/Experimental/UVEditor/Source/UVEditorTools/Private/UVEditorToolUtil.cpp
+++ b/Engine/Plugins/Experimental/UVEditor/Source/UVEditorTools/Private/UVEditorToolUtil.cpp
@@ -26,7 +26,7 @@ void UVEditorToolUtil::GenerateUVUnwrapMesh(
 	{
 		FVector2f UVElement = UVOverlay.GetElement(ElementID);
 		UnwrapMeshOut.InsertVertex(ElementID, UVToVertPosition(UVElement), true);
-		UnwrapMeshUVOverlay->InsertElement(ElementID, (float*)UVElement);
+		UnwrapMeshUVOverlay->InsertElement(ElementID, &UVElement.X);
 	}
 	UnwrapMeshOut.EndUnsafeVerticesInsert();
 	UnwrapMeshUVOverlay->EndUnsafeElementsInsert();
@@ -81,7 +81,7 @@ void UVEditorToolUtil::UpdateUVUnwrapMesh(const FDynamicMeshUVOverlay& UVOverlay
 			}
 			else
 			{
-				UnwrapMeshUVOverlay->InsertElement(ElementID, (float*)ElementValue);
+				UnwrapMeshUVOverlay->InsertElement(ElementID, &ElementValue.X);
 			}
 		}
 	}
@@ -114,7 +114,8 @@ void UVEditorToolUtil::UpdateUVOverlayFromUnwrapMesh(
 		}
 		else
 		{
-			UVOverlayOut.InsertElement(Vid, (float*)VertPositionToUV(UnwrapMeshIn.GetVertex(Vid)));
+			FVector2f UV = VertPositionToUV(UnwrapMeshIn.GetVertex(Vid));
+			UVOverlayOut.InsertElement(Vid, &UV.X);
 		}
 	};
 
diff --git a/Engine/Source/Runtime/GeometryCore/Public/VectorTypes.h b/Engine/Source/Runtime/GeometryCore/Public/VectorTypes.h
--- a/Engine/Source/Runtime/GeometryCore/Public/VectorTypes.h
+++ b/Engine/Source/Runtime/GeometryCore/Public/VectorTypes.h
@@ -38,11 +38,6 @@ struct FVector2
 	{
 	}
 
-	constexpr FVector2(const T* Data)
-		: X(Data[0]), Y(Data[1])
-	{
-	}
-
 	constexpr FVector2(const FVector2& Vec) = default;
 
 	constexpr FVector2(const FVector2D& Vec)
@@ -56,16 +51,6 @@ struct FVector2
 	{
 	}
 
-	explicit constexpr operator const T*() const
-	{
-		return &X;
-	};
-
-	explicit constexpr operator T*()
-	{
-		return &X;
-	}
-
 	explicit operator FVector2D() const
 	{
 		return FVector2D((float)X, (float)Y);
