UVEditor: Fix issue in plugin configuration for runtime vs editor

#rnx

diff --git a/Engine/Plugins/Experimental/UVEditor/UVEditor.uplugin b/Engine/Plugins/Experimental/UVEditor/UVEditor.uplugin
--- a/Engine/Plugins/Experimental/UVEditor/UVEditor.uplugin
+++ b/Engine/Plugins/Experimental/UVEditor/UVEditor.uplugin
@@ -22,7 +22,7 @@
 		},
 		{
 			"Name": "UVEditorTools",
-			"Type": "Runtime",
+			"Type": "Editor",
 			"LoadingPhase": "Default"
 		}
 	],
