UVEditor: Add support for sew action as part of the Select/Transform tool

Additionally adds a new parameter to the DynamicMesh MergeEdge method to ignore orientation discrepancies, useful in 2D mesh manipulation, and fixes a bug in the SetTriangle method where a variable was being left uninitialized.

#rnx

diff --git a/Engine/Plugins/Experimental/UVEditor/Source/UVEditorTools/Private/UVSeamSewAction.cpp b/Engine/Plugins/Experimental/UVEditor/Source/UVEditorTools/Private/UVSeamSewAction.cpp
new file mode 100644
--- /dev/null
+++ b/Engine/Plugins/Experimental/UVEditor/Source/UVEditorTools/Private/UVSeamSewAction.cpp
@@ -0,0 +1,281 @@
+// Copyright Epic Games, Inc. All Rights Reserved.
+
+#include "UVSeamSewAction.h"
+
+#include "Algo/Unique.h"
+#include "Drawing/LineSetComponent.h"
+#include "Drawing/PreviewGeometryActor.h"
+#include "ToolTargets/UVEditorToolMeshInput.h"
+#include "Selection/DynamicMeshSelection.h"
+#include "MeshOpPreviewHelpers.h" // UMeshOpPreviewWithBackgroundCompute
+#include "PreviewMesh.h"
+#include "ToolSetupUtil.h"
+#include "UVToolContextObjects.h"
+
+#define LOCTEXT_NAMESPACE "UUVSeamSewAction"
+
+using namespace UE::Geometry;
+
+
+UUVSeamSewAction::UUVSeamSewAction()
+{
+
+}
+
+void UUVSeamSewAction::Initialize(UWorld* WorldIn, const TArray<TObjectPtr<UUVEditorToolMeshInput>>& TargetsIn)
+{
+	UUVToolAction::Initialize(WorldIn, TargetsIn);
+
+	UnwrapPreviewGeometryActor = World->SpawnActor<APreviewGeometryActor>(FVector::ZeroVector, FRotator(0, 0, 0), FActorSpawnParameters());
+	SewEdgePairingLineSet = NewObject<ULineSetComponent>(UnwrapPreviewGeometryActor);
+	UnwrapPreviewGeometryActor->SetRootComponent(SewEdgePairingLineSet);
+	SewEdgePairingLineSet->RegisterComponent();
+	SewEdgePairingLineSet->SetLineMaterial(ToolSetupUtil::GetDefaultLineComponentMaterial(
+		nullptr, /*bDepthTested*/ true));
+}
+
+void UUVSeamSewAction::Shutdown()
+{
+	if (UnwrapPreviewGeometryActor)
+	{
+		UnwrapPreviewGeometryActor->Destroy();
+		UnwrapPreviewGeometryActor = nullptr;
+	}
+}
+
+
+void UUVSeamSewAction::SetSelection(int32 SelectionTargetIndexIn, const UE::Geometry::FDynamicMeshSelection* NewSelection)
+{
+	SelectionTargetIndex = SelectionTargetIndexIn;
+	CurrentSelection = NewSelection;
+	EdgeSewCandidates.Reset();
+
+	if (SelectionTargetIndex != -1 && CurrentSelection && CurrentSelection->Type == FDynamicMeshSelection::EType::Edge)
+	{
+		for (int32 Eid : CurrentSelection->SelectedIDs)
+		{
+			FIndex2i SewPairCandidate(Eid, FindSewEdgeOppositePairing(Eid));
+			if (SewPairCandidate[1] != IndexConstants::InvalidID)
+			{
+				if (SewPairCandidate[0] > SewPairCandidate[1])
+				{
+					SewPairCandidate.Swap();
+				}
+
+				FEdgePair EdgePair;
+				EdgePair.A = CurrentSelection->Mesh->GetEdgeV(SewPairCandidate[0]);
+				EdgePair.B = CurrentSelection->Mesh->GetEdgeV(SewPairCandidate[1]);
+
+				EdgeSewCandidates.Add(EdgePair);
+			}
+		}
+	}
+
+	UpdateVisualizations();
+}
+
+
+void UUVSeamSewAction::Tick(float DeltaTime) 
+{
+
+
+}
+
+void UUVSeamSewAction::UpdateVisualizations()
+{
+	UpdateSewEdgePreviewLines();
+}
+
+bool UUVSeamSewAction::ExecuteAction(UUVToolEmitChangeAPI& EmitChangeAPI)
+{
+	if (!PreCheckSewAction()) {
+		return false;
+	}
+
+	return ApplySewAction(EmitChangeAPI);
+}
+
+void UUVSeamSewAction::UpdateSewEdgePreviewLines()
+{
+	SewEdgePairingLineSet->Clear();
+	if (CurrentSelection != nullptr && !CurrentSelection->IsEmpty())
+	{
+		FTransform MeshTransform = Targets[SelectionTargetIndex]->UnwrapPreview->PreviewMesh->GetTransform();
+
+		const FDynamicMesh3* UnwrapMesh = Targets[SelectionTargetIndex]->UnwrapPreview->PreviewMesh->GetMesh();
+		for (FEdgePair SewPair : EdgeSewCandidates)
+		{
+			
+			FVector3d Vert1, Vert2;
+			Vert1 = UnwrapMesh->GetVertex(SewPair.A[0]);
+			Vert2 = UnwrapMesh->GetVertex(SewPair.A[1]);
+			UE_LOG(LogGeometry, Warning, TEXT("Red %d %d."), SewPair.A[0], SewPair.A[1]);
+			SewEdgePairingLineSet->AddLine(
+				MeshTransform.TransformPosition(Vert1),
+				MeshTransform.TransformPosition(Vert2),
+				FColor::Red, 5, 3);
+
+			Vert1 = UnwrapMesh->GetVertex(SewPair.B[0]);
+			Vert2 = UnwrapMesh->GetVertex(SewPair.B[1]);
+			UE_LOG(LogGeometry, Warning, TEXT("Green %d %d."), SewPair.B[0], SewPair.B[1]);
+			SewEdgePairingLineSet->AddLine(
+				MeshTransform.TransformPosition(Vert1),
+				MeshTransform.TransformPosition(Vert2),
+				FColor::Green, 5, 3);
+
+		}
+	}
+}
+
+
+bool UUVSeamSewAction::PreCheckSewAction() const
+{
+	if (SelectionTargetIndex == -1 || CurrentSelection == nullptr || CurrentSelection->Mesh == nullptr)
+	{
+		UE_LOG(LogGeometry, Warning, TEXT("Cannot sew UVs. Mesh selection was empty."));
+		return false;
+	}
+
+	if (CurrentSelection->Type != FDynamicMeshSelection::EType::Edge)
+	{
+		UE_LOG(LogGeometry, Warning, TEXT("Cannot sew UVs. Selection was not an edge."));
+		return false;
+	}
+
+	if (EdgeSewCandidates.Num() == 0)
+	{
+		UE_LOG(LogGeometry, Warning, TEXT("Cannot sew UVs. No viable sew candidate edges selected."));
+		return false;
+	}
+
+	return true;
+}
+
+
+int32 UUVSeamSewAction::FindSewEdgeOppositePairing(int32 UnwrapEid) const
+{
+	// Given an edge id on the unwrap mesh, determine it's opposite edge suitable for UV sewing elsewhere on the unwrap mesh
+
+	if (!CurrentSelection->Mesh->IsBoundaryEdge(UnwrapEid))
+	{
+		UE_LOG(LogGeometry, Warning, TEXT("Selected edge is unsuitable for sewing - Is not a boundary edge in UV space"));
+		return IndexConstants::InvalidID;
+	}
+	int32 UVLayerIndex = Targets[SelectionTargetIndex]->UVLayerIndex;
+	FDynamicMeshUVOverlay* UVOverlay = Targets[SelectionTargetIndex]->AppliedCanonical.Get()->Attributes()->GetUVLayer(UVLayerIndex);
+	FIndex2i EdgeVids = CurrentSelection->Mesh->GetEdgeV(UnwrapEid);
+	int32 ParentVid_0 = UVOverlay->GetParentVertex(EdgeVids[0]);
+	int32 ParentVid_1 = UVOverlay->GetParentVertex(EdgeVids[1]);
+
+	TArray<int32> Vid0Elements;
+	TArray<int32> Vid1Elements;
+	UVOverlay->GetVertexElements(ParentVid_0, Vid0Elements);
+	UVOverlay->GetVertexElements(ParentVid_1, Vid1Elements);
+
+	// Now that we know our potential elements, locate one single edge that is a viable pairing, or return InvalidID if there are none or more.
+	int32 OppositeEid = IndexConstants::InvalidID;
+	for (int32 Vid0 : Vid0Elements)
+	{
+		for (int32 Vid1 : Vid1Elements)
+		{
+			int32 EidCandidate = CurrentSelection->Mesh->FindEdge(Vid0, Vid1);
+			if (EidCandidate == UnwrapEid || EidCandidate == IndexConstants::InvalidID)
+			{
+				continue;
+			}
+			if (OppositeEid == IndexConstants::InvalidID)
+			{
+				OppositeEid = EidCandidate;
+				int32 ParentCandidateVid_0 = UVOverlay->GetParentVertex(Vid0);
+				int32 ParentCandidateVid_1 = UVOverlay->GetParentVertex(Vid1);
+				check(ParentCandidateVid_0 == ParentVid_0 && ParentCandidateVid_1 == ParentVid_1)
+			}
+			else
+			{
+				UE_LOG(LogGeometry, Warning, TEXT("Selected edge is unsuitable for sewing - there are non-manifold candidate edges."));
+				return IndexConstants::InvalidID;
+			}
+		}
+	}
+
+	return OppositeEid;
+}
+
+
+bool UUVSeamSewAction::ApplySewAction(UUVToolEmitChangeAPI& EmitChangeAPI)
+{
+	FDynamicMesh3& MeshToSew = *(Targets[SelectionTargetIndex]->UnwrapCanonical);
+
+	TArray<int32> SelectedTids;
+	TArray<FIndex2i> ResolvedEdgePairs;
+	for (FEdgePair SewPair : EdgeSewCandidates)
+	{
+		ResolvedEdgePairs.Add(FIndex2i(MeshToSew.FindEdge(SewPair.A[0], SewPair.A[1]),
+			                           MeshToSew.FindEdge(SewPair.B[0], SewPair.B[1])));
+
+		for (int32 VertexIndex = 0; VertexIndex < 2; ++VertexIndex)
+		{
+			TArray<int32> Tids;
+			MeshToSew.GetVtxTriangles(SewPair.A[VertexIndex], Tids);
+			SelectedTids.Append(Tids);
+		}			
+		
+		for (int32 VertexIndex = 0; VertexIndex < 2; ++VertexIndex)
+		{
+			TArray<int32> Tids;
+			MeshToSew.GetVtxTriangles(SewPair.B[VertexIndex], Tids);
+			SelectedTids.Append(Tids);
+		}
+	}
+	SelectedTids.Sort();
+	Algo::Unique(SelectedTids);
+
+	FDynamicMeshChangeTracker ChangeTracker(Targets[SelectionTargetIndex]->UnwrapCanonical.Get());
+	ChangeTracker.BeginChange();
+	ChangeTracker.SaveTriangles(SelectedTids, true);
+
+	EdgeSewCandidates.Sort();
+	Algo::Unique(EdgeSewCandidates);
+
+	TArray<FDynamicMesh3::FMergeEdgesInfo> AllMergeInfo;
+	
+	for (FIndex2i EdgePair : ResolvedEdgePairs)
+	{
+		FDynamicMesh3::FMergeEdgesInfo MergeInfo;
+		EMeshResult Result = MeshToSew.MergeEdges(EdgePair[0], EdgePair[1], MergeInfo, false);
+		if (Result == EMeshResult::Ok)
+		{
+			AllMergeInfo.Add(MergeInfo);
+		}
+		else
+		{
+			UE_LOG(LogGeometry, Warning, TEXT("Failed to sew edge pair %d / %d. Failed with code %d"), EdgePair[0], EdgePair[1], Result);
+		}
+	}
+	checkSlow(MeshToSew.CheckValidity());
+
+
+	TArray<int32> RemainingVids;
+	RemainingVids.SetNum(AllMergeInfo.Num() * 2);
+	int32 RemainingVidsIndex = 0;
+	for (FDynamicMesh3::FMergeEdgesInfo& MergeInfo : AllMergeInfo) {
+		RemainingVids[RemainingVidsIndex++] = MergeInfo.KeptVerts[0];
+		RemainingVids[RemainingVidsIndex++] = MergeInfo.KeptVerts[1];
+	}
+
+	Targets[SelectionTargetIndex]->UpdateUnwrapCanonicalOverlayFromPositions(&RemainingVids, &SelectedTids);
+	Targets[SelectionTargetIndex]->UpdateAllFromUnwrapCanonical(&RemainingVids, &SelectedTids, &SelectedTids);
+	checkSlow(MeshToSew.IsSameAs(*Targets[SelectionTargetIndex]->UnwrapPreview->PreviewMesh->GetMesh(), FDynamicMesh3::FSameAsOptions()));
+
+	const FText TransactionName(LOCTEXT("SewCompleteTransactionName", "Sew Edges"));
+	EmitChangeAPI.BeginUndoTransaction(TransactionName);
+	EmitChangeAPI.EmitToolIndependentUnwrapCanonicalChange(Targets[SelectionTargetIndex],
+		ChangeTracker.EndChange(), TransactionName);
+	EmitChangeAPI.EndUndoTransaction();
+
+	SetSelection(-1, nullptr);
+
+	return true;
+}
+
+#undef LOCTEXT_NAMESPACE
\ No newline at end of file
diff --git a/Engine/Plugins/Experimental/UVEditor/Source/UVEditorTools/Private/UVSelectTool.cpp b/Engine/Plugins/Experimental/UVEditor/Source/UVEditorTools/Private/UVSelectTool.cpp
--- a/Engine/Plugins/Experimental/UVEditor/Source/UVEditorTools/Private/UVSelectTool.cpp
+++ b/Engine/Plugins/Experimental/UVEditor/Source/UVEditorTools/Private/UVSelectTool.cpp
@@ -2,6 +2,7 @@
 
 #include "UVSelectTool.h"
 
+#include "Algo/Unique.h"
 #include "BaseGizmos/TransformGizmo.h"
 #include "ContextObjectStore.h"
 #include "DynamicMesh/DynamicMeshChangeTracker.h"
@@ -9,6 +10,7 @@
 #include "Drawing/MeshElementsVisualizer.h"
 #include "Drawing/PreviewGeometryActor.h"
 #include "DynamicMesh/DynamicMeshChangeTracker.h"
+#include "DynamicMesh/MeshIndexUtil.h"
 #include "ToolTargets/UVEditorToolMeshInput.h"
 #include "InteractiveToolManager.h"
 #include "MeshOpPreviewHelpers.h" // UMeshOpPreviewWithBackgroundCompute
@@ -17,6 +19,8 @@
 #include "Selection/DynamicMeshSelection.h"
 #include "ToolSetupUtil.h"
 
+#include "UVSeamSewAction.h"
+
 #include "ToolTargetManager.h"
 
 #define LOCTEXT_NAMESPACE "UUVSelectTool"
@@ -212,6 +216,7 @@ namespace UVSelectToolLocals
 		FTransform GizmoBefore;
 		FTransform GizmoAfter;
 	};
+
 }
 
 void UUVSelectToolSpeculativeChangeAPI::EmitSpeculativeChange(UObject* TargetObject, 
@@ -258,12 +263,28 @@ UInteractiveTool* UUVSelectToolBuilder::BuildTool(const FToolBuilderState& Scene
 	return NewTool;
 }
 
+// Tool property functions
+
+void  USelectToolActionPropertySet::Sew()
+{
+	PostAction(ESelectToolAction::Sew); 
+}
+
+void USelectToolActionPropertySet::PostAction(ESelectToolAction Action)
+{
+	if (ParentTool.IsValid())
+	{
+		ParentTool->RequestAction(Action);
+	}
+}
+
+
 void UUVSelectTool::Setup()
 {
 	check(Targets.Num() > 0);
 
 	UInteractiveTool::Setup();
-
+	
 	SetToolDisplayName(LOCTEXT("ToolName", "UV Select Tool"));
 
 	Settings = NewObject<UUVSelectToolProperties>(this);
@@ -273,6 +294,10 @@ void UUVSelectTool::Setup()
 	UContextObjectStore* ContextStore = GetToolManager()->GetContextObjectStore();
 	EmitChangeAPI = ContextStore->FindContext<UUVToolEmitChangeAPI>();
 
+	ToolActions = NewObject<USelectToolActionPropertySet>(this);
+	ToolActions->Initialize(this);
+	AddToolPropertySource(ToolActions);
+
 	SelectionMechanic = NewObject<UMeshSelectionMechanic>();
 	SelectionMechanic->Setup(this);
 	SelectionMechanic->SetWorld(Targets[0]->UnwrapPreview->GetWorld());
@@ -392,6 +417,13 @@ void UUVSelectTool::Setup()
 	LivePreviewLineSet->SetLineMaterial(ToolSetupUtil::GetDefaultLineComponentMaterial(
 		GetToolManager(), /*bDepthTested*/ true));
 
+	SewAction = NewObject<UUVSeamSewAction>();
+	SewAction->Setup(this);
+	SewAction->Initialize(
+		Targets[0]->UnwrapPreview->GetWorld(),
+		Targets);
+
+
 	if (!SelectionMechanic->GetCurrentSelection().IsEmpty())
 	{
 		OnSelectionChanged();
@@ -457,6 +489,11 @@ void UUVSelectTool::Shutdown(EToolShutdownType ShutdownType)
 		LivePreviewGeometryActor = nullptr;
 	}
 
+	if (SewAction)
+	{
+		SewAction->Shutdown();
+	}
+
 	// Calls shutdown on gizmo and destroys it.
 	GetToolManager()->GetPairedGizmoManager()->DestroyAllGizmosByOwner(this);
 
@@ -527,6 +564,7 @@ void UUVSelectTool::OnSelectionChanged()
 	MovingVids.Reset();
 	SelectedTids.Reset();
 	BoundaryEids.Reset();
+
 	if (!Selection.IsEmpty())
 	{
 		// Note which mesh we're selecting in.
@@ -603,8 +641,8 @@ void UUVSelectTool::OnSelectionChanged()
 						}
 					}
 				}
-				}
 			}
+		}
 		else if (Selection.Type == FDynamicMeshSelection::EType::Vertex)
 		{
 			for (int32 Vid : Selection.SelectedIDs)
@@ -634,8 +672,9 @@ void UUVSelectTool::OnSelectionChanged()
 		}
 	}
 
+	SewAction->SetSelection(SelectionTargetIndex, &Selection);
+
 	UpdateLivePreviewLines();
-	
 	UpdateGizmo();
 }
 
@@ -658,11 +697,12 @@ void UUVSelectTool::UpdateLivePreviewLines()
 				MeshTransform.TransformPosition(Vert1), 
 				MeshTransform.TransformPosition(Vert2), 
 				FColor::Yellow, 2, 1.5);
-		}
-		
+		}	
 	}
 }
 
+
+
 void UUVSelectTool::SetGizmoEnabled(bool bEnabledIn)
 {
 	bGizmoEnabled = bEnabledIn;
@@ -783,6 +823,7 @@ void UUVSelectTool::ApplyGizmoTransform()
 		}
 
 		bGizmoTransformNeedsApplication = false;
+		SewAction->UpdateVisualizations();
 	}
 }
 
@@ -794,6 +835,49 @@ void UUVSelectTool::Render(IToolsContextRenderAPI* RenderAPI)
 void UUVSelectTool::OnTick(float DeltaTime)
 {
 	ApplyGizmoTransform();
+
+	// Deal with any buttons that may have been clicked
+	if (PendingAction != ESelectToolAction::NoAction)
+	{
+		ApplyAction(PendingAction);
+		PendingAction = ESelectToolAction::NoAction;
+	}
+}
+
+void UUVSelectTool::RequestAction(ESelectToolAction ActionType)
+{
+	if (PendingAction == ESelectToolAction::NoAction)
+	{
+		PendingAction = ActionType;
+	}
+}
+
+void UUVSelectTool::ApplyAction(ESelectToolAction ActionType)
+{
+	switch (ActionType)
+	{
+	case ESelectToolAction::Sew:
+		if (SewAction)
+		{
+			bool ActionSuccessful = SewAction->ExecuteAction(*EmitChangeAPI);
+
+			if (ActionSuccessful)
+			{
+				if (!AABBTrees[SelectionTargetIndex]->IsValid(false))
+				{
+					AABBTrees[SelectionTargetIndex]->Build();
+				}
+
+				// TODO: Emit a selection update instead of just clearing things here. If we do
+				// still clear, remember that SelectionTargetIndex gets updated, so either do it
+				// here or save a pointer to input object before doing stuff.
+				SelectionMechanic->SetSelection(FDynamicMeshSelection(), true, false);
+			}
+		}
+		break;
+	default:
+		break;
+	}
 }
 
 
diff --git a/Engine/Plugins/Experimental/UVEditor/Source/UVEditorTools/Public/UVSeamSewAction.h b/Engine/Plugins/Experimental/UVEditor/Source/UVEditorTools/Public/UVSeamSewAction.h
new file mode 100644
--- /dev/null
+++ b/Engine/Plugins/Experimental/UVEditor/Source/UVEditorTools/Public/UVSeamSewAction.h
@@ -0,0 +1,69 @@
+// Copyright Epic Games, Inc. All Rights Reserved.
+
+#pragma once
+
+#include "CoreMinimal.h"
+#include "GeometryBase.h"
+#include "IndexTypes.h"
+
+#include "UVToolAction.h"
+
+#include "UVSeamSewAction.generated.h"
+
+PREDECLARE_GEOMETRY(class FDynamicMeshSelection);
+class APreviewGeometryActor;
+class ULineSetComponent;
+
+// We store edges by pairs of Vids here, since Edge Ids seem unreliable between meshes
+// when Undo/Redos are happening and modifiying the canonical and preview meshes
+typedef UE::Geometry::FIndex2i FEdgeByVids;
+struct FEdgePair {
+	FEdgeByVids A;
+	FEdgeByVids B;
+
+	bool operator==(const FEdgePair& Other) const 
+	{
+		return A == Other.A && B == Other.B;
+	}
+	bool operator<(const FEdgePair& Other) const
+	{
+		return ((A[0] < Other.A[0]) || (A[0] == Other.A[0] && A[1] < Other.A[1])) ||
+			((A == Other.A) && ((B[0] < Other.B[0]) || (B[0] == Other.B[0] && B[1] < Other.B[1])));
+	}
+};
+
+UCLASS()
+class UVEDITORTOOLS_API UUVSeamSewAction : public UUVToolAction
+{	
+	GENERATED_BODY()
+
+public:
+
+	UUVSeamSewAction();
+
+	void Initialize(UWorld* WorldIn, const TArray<TObjectPtr<UUVEditorToolMeshInput>>& TargetsIn) override;
+	void Shutdown() override;
+
+	void SetSelection(int32 SelectionTargetIndexIn, const UE::Geometry::FDynamicMeshSelection* NewSelection);
+	virtual void Tick(float DeltaTime) override;
+	virtual bool ExecuteAction(UUVToolEmitChangeAPI& EmitChangeAPI) override;
+	virtual void UpdateVisualizations() override;
+
+protected:
+
+	TArray<FEdgePair> EdgeSewCandidates;
+
+	UPROPERTY()
+	TObjectPtr<APreviewGeometryActor> UnwrapPreviewGeometryActor = nullptr;
+
+	UPROPERTY()
+	TObjectPtr<ULineSetComponent> SewEdgePairingLineSet = nullptr;
+
+	int32 SelectionTargetIndex;
+	const UE::Geometry::FDynamicMeshSelection* CurrentSelection;
+
+	int32 FindSewEdgeOppositePairing(int32 UnwrapEid) const;
+	bool PreCheckSewAction() const;
+	bool ApplySewAction(UUVToolEmitChangeAPI& EmitChangeAPI);
+	void UpdateSewEdgePreviewLines();
+};
\ No newline at end of file
diff --git a/Engine/Plugins/Experimental/UVEditor/Source/UVEditorTools/Public/UVSelectTool.h b/Engine/Plugins/Experimental/UVEditor/Source/UVEditorTools/Public/UVSelectTool.h
--- a/Engine/Plugins/Experimental/UVEditor/Source/UVEditorTools/Public/UVSelectTool.h
+++ b/Engine/Plugins/Experimental/UVEditor/Source/UVEditorTools/Public/UVSelectTool.h
@@ -11,6 +11,7 @@
 #include "InteractiveToolBuilder.h"
 #include "TargetInterfaces/UVUnwrapDynamicMesh.h"
 #include "UVToolContextObjects.h"
+#include "IndexTypes.h"
 
 #include "UVSelectTool.generated.h"
 
@@ -28,6 +29,7 @@ class UUVToolStateObjectStore;
 class UPreviewMesh;
 class UUVEditorToolMeshInput;
 class UUVToolEmitChangeAPI;
+class UUVSeamSewAction;
 
 UCLASS()
 class UVEDITORTOOLS_API UUVSelectToolBuilder : public UInteractiveToolBuilder
@@ -73,6 +75,35 @@ public:
 	bool bUpdatePreviewDuringDrag = true;
 };
 
+
+
+UENUM()
+enum class ESelectToolAction
+{
+	NoAction,
+
+	Sew,
+};
+
+UCLASS()
+class UVEDITORTOOLS_API USelectToolActionPropertySet : public UInteractiveToolPropertySet
+{
+	GENERATED_BODY()
+
+public:
+	TWeakObjectPtr<UUVSelectTool> ParentTool;
+
+	void Initialize(UUVSelectTool* ParentToolIn) { ParentTool = ParentToolIn; }
+	void PostAction(ESelectToolAction Action);
+
+	/** Move the mirror plane to center of bounding box without changing its normal. */
+	UFUNCTION(CallInEditor, Category = Actions)
+	void Sew();
+
+};
+
+
+
 /**
  * A tool for selecting elements of a flat FDynamicMesh corresponding to a UV layer of some asset.
  * If bGizmoEnabled is set to true, the selected elements can be moved around.
@@ -111,6 +142,8 @@ public:
 	virtual bool HasCancel() const override { return false; }
 	virtual bool HasAccept() const override { return false; }
 
+	void RequestAction(ESelectToolAction ActionType);
+
 protected:
 	virtual void OnSelectionChanged();
 
@@ -124,7 +157,9 @@ protected:
 	virtual void UpdateGizmo();
 	virtual void UpdateLivePreviewLines();
 
-	void ConfigureSelectionModeFromControls();
+	void SetToolActionsVisiblity(bool isVisible);
+
+	void ConfigureSelectionModeFromControls();	
 
 	UWorld* TargetWorld;
 
@@ -134,6 +169,9 @@ protected:
 	UPROPERTY()
 	TObjectPtr<UUVSelectToolProperties> Settings = nullptr;
 
+	UPROPERTY()
+	TObjectPtr<USelectToolActionPropertySet> ToolActions = nullptr;
+
 	UPROPERTY()
 	TObjectPtr<UMeshSelectionMechanic> SelectionMechanic = nullptr;
 
@@ -154,16 +192,24 @@ protected:
 	UPROPERTY()
 	TObjectPtr<UUVSelectToolChangeRouter> ChangeRouter = nullptr;
 
+	UPROPERTY()
+	TObjectPtr<UUVSeamSewAction> SewAction = nullptr;
+
 	UE::Geometry::FFrame3d InitialGizmoFrame;
 	FTransform UnappliedGizmoTransform;
 	bool bInDrag = false;
 	bool bGizmoTransformNeedsApplication = false;
 
 	TArray<int32> MovingVids;
+	TArray<int32> SelectedVids;
 	TArray<int32> SelectedTids;
 	TArray<FVector3d> MovingVertOriginalPositions;
 	int32 SelectionTargetIndex;
 	TArray<int32> BoundaryEids;
+	
+	ESelectToolAction PendingAction;
+
+	void ApplyAction(ESelectToolAction ActionType);
 
 	// We need this flag so that SetGizmoVisibility can be called before Setup() by the tool builder.
 	bool bGizmoEnabled = false;
diff --git a/Engine/Plugins/Experimental/UVEditor/Source/UVEditorTools/Public/UVToolAction.h b/Engine/Plugins/Experimental/UVEditor/Source/UVEditorTools/Public/UVToolAction.h
new file mode 100644
--- /dev/null
+++ b/Engine/Plugins/Experimental/UVEditor/Source/UVEditorTools/Public/UVToolAction.h
@@ -0,0 +1,46 @@
+// Copyright Epic Games, Inc. All Rights Reserved.
+
+#pragma once
+
+#include "CoreMinimal.h"
+#include "InteractionMechanic.h"
+
+#include "UVToolAction.generated.h"
+
+class UUVEditorToolMeshInput;
+class UUVToolEmitChangeAPI;
+
+
+// TODO: Should this be spread out across multiple files?
+
+UCLASS()
+class UVEDITORTOOLS_API UUVToolAction : public UInteractionMechanic
+{
+	GENERATED_BODY()
+public:
+	virtual void Initialize(UWorld* WorldIn, const TArray<TObjectPtr<UUVEditorToolMeshInput>>& TargetsIn)
+	{
+		World = WorldIn;
+		Targets = TargetsIn;
+	}		
+	virtual void Shutdown() override {
+		UInteractionMechanic::Shutdown();
+
+		World = nullptr;
+
+		for(TObjectPtr<UUVEditorToolMeshInput>& Target : Targets)
+		{
+			Target = nullptr;
+		}
+	}
+	virtual bool ExecuteAction(UUVToolEmitChangeAPI& EmitChangeAPI) { return true; };
+	virtual void UpdateVisualizations() {};
+
+protected:
+
+	UPROPERTY()
+	TWeakObjectPtr<UWorld> World;
+
+	UPROPERTY()
+	TArray<TObjectPtr<UUVEditorToolMeshInput>> Targets;
+};
diff --git a/Engine/Source/Runtime/GeometryCore/Private/DynamicMesh/DynamicMesh3_Edits.cpp b/Engine/Source/Runtime/GeometryCore/Private/DynamicMesh/DynamicMesh3_Edits.cpp
--- a/Engine/Source/Runtime/GeometryCore/Private/DynamicMesh/DynamicMesh3_Edits.cpp
+++ b/Engine/Source/Runtime/GeometryCore/Private/DynamicMesh/DynamicMesh3_Edits.cpp
@@ -777,7 +777,7 @@ EMeshResult FDynamicMesh3::SetTriangle(int tID, const FIndex3i& newv, bool bRemo
 			continue;
 		}
 		ReplaceEdgeTriangle(eid, tID, InvalidID);
-		const FEdge Edge;
+		const FEdge Edge = GetEdge(eid);
 		if (Edge.Tri[0] == InvalidID)
 		{
 			int a = Edge.Vert[0];
@@ -1617,7 +1617,7 @@ EMeshResult FDynamicMesh3::CollapseEdge(int vKeep, int vRemove, double collapse_
 
 
 
-EMeshResult FDynamicMesh3::MergeEdges(int eKeep, int eDiscard, FMergeEdgesInfo& MergeInfo)
+EMeshResult FDynamicMesh3::MergeEdges(int eKeep, int eDiscard, FMergeEdgesInfo& MergeInfo, bool bCheckValidOrientation)
 {
 	MergeInfo = FMergeEdgesInfo();
 
@@ -1649,7 +1649,8 @@ EMeshResult FDynamicMesh3::MergeEdges(int eKeep, int eDiscard, FMergeEdgesInfo&
 	IndexUtil::OrientTriEdge(c, d, GetTriangle(tcd));
 	int x = c; c = d; d = x;   // joinable bdry edges have opposing orientations, so flip to get ac and b/d correspondences
 	FVector3d Va = GetVertex(a), Vb = GetVertex(b), Vc = GetVertex(c), Vd = GetVertex(d);
-	if (((Va - Vc).SquaredLength() + (Vb - Vd).SquaredLength()) >
+	if (bCheckValidOrientation && 
+		((Va - Vc).SquaredLength() + (Vb - Vd).SquaredLength()) >
 		((Va - Vd).SquaredLength() + (Vb - Vc).SquaredLength()))
 	{
 		return EMeshResult::Failed_SameOrientation;
diff --git a/Engine/Source/Runtime/GeometryCore/Public/DynamicMesh/DynamicMesh3.h b/Engine/Source/Runtime/GeometryCore/Public/DynamicMesh/DynamicMesh3.h
--- a/Engine/Source/Runtime/GeometryCore/Public/DynamicMesh/DynamicMesh3.h
+++ b/Engine/Source/Runtime/GeometryCore/Public/DynamicMesh/DynamicMesh3.h
@@ -1225,9 +1225,10 @@ public:
 	 * @param KeepEdgeID index of the edge that should be kept
 	 * @param DiscardEdgeID index of the edge that should be removed
 	 * @param MergeInfo returned information about new and modified mesh elements
+	 * @param CheckValidOrientation perform edge consistency orientation checks before merging.
 	 * @return Ok on success, or enum value indicates why operation cannot be applied. Mesh remains unmodified on error.
 	 */
-	virtual EMeshResult MergeEdges(int KeepEdgeID, int DiscardEdgeID, FMergeEdgesInfo& MergeInfo);
+	virtual EMeshResult MergeEdges(int KeepEdgeID, int DiscardEdgeID, FMergeEdgesInfo& MergeInfo, bool bCheckValidOrientation=true);
 
 
 	/**
