UVEditor: Fix compile error on NonUnity builds

#rnx

diff --git a/Engine/Plugins/Experimental/UVEditor/Source/UVEditorTools/Private/UVSeamSewAction.cpp b/Engine/Plugins/Experimental/UVEditor/Source/UVEditorTools/Private/UVSeamSewAction.cpp
--- a/Engine/Plugins/Experimental/UVEditor/Source/UVEditorTools/Private/UVSeamSewAction.cpp
+++ b/Engine/Plugins/Experimental/UVEditor/Source/UVEditorTools/Private/UVSeamSewAction.cpp
@@ -22,16 +22,41 @@ UUVSeamSewAction::UUVSeamSewAction()
 
 }
 
-void UUVSeamSewAction::Initialize(UWorld* WorldIn, const TArray<TObjectPtr<UUVEditorToolMeshInput>>& TargetsIn)
+void UUVSeamSewAction::SetWorld(UWorld* WorldIn)
 {
-	UUVToolAction::Initialize(WorldIn, TargetsIn);
+	UUVToolAction::SetWorld(WorldIn);
 
-	UnwrapPreviewGeometryActor = World->SpawnActor<APreviewGeometryActor>(FVector::ZeroVector, FRotator(0, 0, 0), FActorSpawnParameters());
-	SewEdgePairingLineSet = NewObject<ULineSetComponent>(UnwrapPreviewGeometryActor);
+	if (UnwrapPreviewGeometryActor)
+	{
+		UnwrapPreviewGeometryActor->Destroy();
+	}
+
+	// We need the world so we can create the geometry actor in the right place.
+	FRotator Rotation(0.0f, 0.0f, 0.0f);
+	FActorSpawnParameters SpawnInfo;
+	UnwrapPreviewGeometryActor = WorldIn->SpawnActor<APreviewGeometryActor>(FVector::ZeroVector, Rotation, SpawnInfo);
+
+	// Attach the rendering component to the actor
+	SewEdgePairingLineSet->Rename(nullptr, UnwrapPreviewGeometryActor); // Changes the "outer"
 	UnwrapPreviewGeometryActor->SetRootComponent(SewEdgePairingLineSet);
-	SewEdgePairingLineSet->RegisterComponent();
+	if (SewEdgePairingLineSet->IsRegistered())
+	{
+		SewEdgePairingLineSet->ReregisterComponent();
+	}
+	else
+	{
+		SewEdgePairingLineSet->RegisterComponent();
+	}
+
+}
+
+void UUVSeamSewAction::Setup(UInteractiveTool* ParentToolIn)
+{
+	UUVToolAction::Setup(ParentToolIn);
+
+	SewEdgePairingLineSet = NewObject<ULineSetComponent>();
 	SewEdgePairingLineSet->SetLineMaterial(ToolSetupUtil::GetDefaultLineComponentMaterial(
-		nullptr, /*bDepthTested*/ true));
+		GetParentTool()->GetToolManager(), /*bDepthTested*/ true));
 }
 
 void UUVSeamSewAction::Shutdown()
diff --git a/Engine/Plugins/Experimental/UVEditor/Source/UVEditorTools/Private/UVSelectTool.cpp b/Engine/Plugins/Experimental/UVEditor/Source/UVEditorTools/Private/UVSelectTool.cpp
--- a/Engine/Plugins/Experimental/UVEditor/Source/UVEditorTools/Private/UVSelectTool.cpp
+++ b/Engine/Plugins/Experimental/UVEditor/Source/UVEditorTools/Private/UVSelectTool.cpp
@@ -419,9 +419,8 @@ void UUVSelectTool::Setup()
 
 	SewAction = NewObject<UUVSeamSewAction>();
 	SewAction->Setup(this);
-	SewAction->Initialize(
-		Targets[0]->UnwrapPreview->GetWorld(),
-		Targets);
+	SewAction->SetTargets(Targets);
+	SewAction->SetWorld(Targets[0]->UnwrapPreview->GetWorld());
 
 
 	if (!SelectionMechanic->GetCurrentSelection().IsEmpty())
diff --git a/Engine/Plugins/Experimental/UVEditor/Source/UVEditorTools/Public/UVSeamSewAction.h b/Engine/Plugins/Experimental/UVEditor/Source/UVEditorTools/Public/UVSeamSewAction.h
--- a/Engine/Plugins/Experimental/UVEditor/Source/UVEditorTools/Public/UVSeamSewAction.h
+++ b/Engine/Plugins/Experimental/UVEditor/Source/UVEditorTools/Public/UVSeamSewAction.h
@@ -41,7 +41,8 @@ public:
 
 	UUVSeamSewAction();
 
-	void Initialize(UWorld* WorldIn, const TArray<TObjectPtr<UUVEditorToolMeshInput>>& TargetsIn) override;
+	void SetWorld(UWorld* WorldIn) override;
+	void Setup(UInteractiveTool* ParentToolIn) override;
 	void Shutdown() override;
 
 	void SetSelection(int32 SelectionTargetIndexIn, const UE::Geometry::FDynamicMeshSelection* NewSelection);
diff --git a/Engine/Plugins/Experimental/UVEditor/Source/UVEditorTools/Public/UVToolAction.h b/Engine/Plugins/Experimental/UVEditor/Source/UVEditorTools/Public/UVToolAction.h
--- a/Engine/Plugins/Experimental/UVEditor/Source/UVEditorTools/Public/UVToolAction.h
+++ b/Engine/Plugins/Experimental/UVEditor/Source/UVEditorTools/Public/UVToolAction.h
@@ -18,16 +18,16 @@ class UVEDITORTOOLS_API UUVToolAction : public UInteractionMechanic
 {
 	GENERATED_BODY()
 public:
-	virtual void Initialize(UWorld* WorldIn, const TArray<TObjectPtr<UUVEditorToolMeshInput>>& TargetsIn)
+	virtual void SetWorld(UWorld* WorldIn) {}
+
+	virtual void SetTargets(const TArray<TObjectPtr<UUVEditorToolMeshInput>>& TargetsIn)
 	{
-		World = WorldIn;
 		Targets = TargetsIn;
-	}		
+	}
+
 	virtual void Shutdown() override {
 		UInteractionMechanic::Shutdown();
 
-		World = nullptr;
-
 		for(TObjectPtr<UUVEditorToolMeshInput>& Target : Targets)
 		{
 			Target = nullptr;
@@ -38,9 +38,6 @@ public:
 
 protected:
 
-	UPROPERTY()
-	TWeakObjectPtr<UWorld> World;
-
 	UPROPERTY()
 	TArray<TObjectPtr<UUVEditorToolMeshInput>> Targets;
 };
