UVEditor: Clean up depth offset values for various scene elements

#rnx

diff --git a/Engine/Plugins/Experimental/MeshModelingToolsetExp/Source/ModelingComponents/Public/Selection/MeshSelectionMechanic.h b/Engine/Plugins/Experimental/MeshModelingToolsetExp/Source/ModelingComponents/Public/Selection/MeshSelectionMechanic.h
--- a/Engine/Plugins/Experimental/MeshModelingToolsetExp/Source/ModelingComponents/Public/Selection/MeshSelectionMechanic.h
+++ b/Engine/Plugins/Experimental/MeshModelingToolsetExp/Source/ModelingComponents/Public/Selection/MeshSelectionMechanic.h
@@ -129,8 +129,7 @@ protected:
 
 	float LineThickness = 3;
 	float PointThickness = 6;
-	// TODO: This should probably be lower, but we have to lower the depth bias elsewhere first.
-	float DepthBias = 5;
+	float DepthBias = 0.7;
 
 	FVector3d CurrentSelectionCentroid;
 	int32 CentroidTimestamp = -1;
diff --git a/Engine/Plugins/Experimental/UVEditor/Source/UVEditor/Private/UVEditor2DViewportClient.cpp b/Engine/Plugins/Experimental/UVEditor/Source/UVEditor/Private/UVEditor2DViewportClient.cpp
--- a/Engine/Plugins/Experimental/UVEditor/Source/UVEditor/Private/UVEditor2DViewportClient.cpp
+++ b/Engine/Plugins/Experimental/UVEditor/Source/UVEditor/Private/UVEditor2DViewportClient.cpp
@@ -97,7 +97,7 @@ void FUVEditor2DViewportClient::Draw(const FSceneView* View, FPrimitiveDrawInter
 	FVector WorldBoundsMax = View->ScreenToWorld(MaxScreen);
 	FVector WorldBoundsMin = View->ScreenToWorld(MinScreen);
 	FVector ViewLoc = GetViewLocation();
-	ViewLoc.Z = 0.01; // We are treating the scene like a 2D plane, so we'll clamp the Z position here to 
+	ViewLoc.Z = 0.0; // We are treating the scene like a 2D plane, so we'll clamp the Z position here to 
 	               // 0 as a simple projection step just in case.
 
 	// Prevent grid from drawing if we are too close or too far, in order to avoid potential graphical issues.
@@ -122,8 +122,8 @@ void FUVEditor2DViewportClient::Draw(const FSceneView* View, FPrimitiveDrawInter
 	float AxisExtent = FMathf::Max(UVScale, FMathf::Min(WorldBoundsMax.Y, WorldBoundsMax.X));
 
 	// Draw colored axis lines
-	PDI->DrawLine(FVector(0, 0, 0.01), FVector(AxisExtent, 0, 0), FLinearColor::Red, SDPG_World, AxisThickness, 0, true);
-	PDI->DrawLine(FVector(0, 0, 0.01), FVector(0, AxisExtent, 0), FLinearColor::Green, SDPG_World, AxisThickness, 0, true);
+	PDI->DrawLine(FVector(0, 0, 0), FVector(AxisExtent, 0, 0), FLinearColor::Red, SDPG_World, AxisThickness, 0, true);
+	PDI->DrawLine(FVector(0, 0, 0), FVector(0, AxisExtent, 0), FLinearColor::Green, SDPG_World, AxisThickness, 0, true);
 
 	// TODO: Draw a little UV axis thing in the lower left, like the XYZ things that normal viewports have.
 
diff --git a/Engine/Plugins/Experimental/UVEditor/Source/UVEditor/Private/UVEditorMode.cpp b/Engine/Plugins/Experimental/UVEditor/Source/UVEditor/Private/UVEditorMode.cpp
--- a/Engine/Plugins/Experimental/UVEditor/Source/UVEditor/Private/UVEditorMode.cpp
+++ b/Engine/Plugins/Experimental/UVEditor/Source/UVEditor/Private/UVEditorMode.cpp
@@ -217,7 +217,7 @@ void UUVEditorMode::InitializeTargets(TArray<TObjectPtr<UObject>>& AssetsIn, TAr
 			0, ToolSetupUtil::GetCustomTwoSidedDepthOffsetMaterial(
 				GetToolManager(),
 				(FLinearColor)TriangleColor,
-				2,  //depth offset
+				TriangleDepthOffset,
 				TriangleOpacity));
 
 		// Initialize our timestamp so we can later detect changes
@@ -227,7 +227,7 @@ void UUVEditorMode::InitializeTargets(TArray<TObjectPtr<UObject>>& AssetsIn, TAr
 		UMeshElementsVisualizer* WireframeDisplay = NewObject<UMeshElementsVisualizer>(this);
 		WireframeDisplay->CreateInWorld(GetWorld(), FTransform::Identity);
 
-		WireframeDisplay->Settings->DepthBias = 2.5;
+		WireframeDisplay->Settings->DepthBias = WireframeDepthOffset;
 		WireframeDisplay->Settings->bAdjustDepthBiasUsingMeshSize = false;
 		WireframeDisplay->Settings->bShowWireframe = true;
 		WireframeDisplay->Settings->bShowBorders = true;
@@ -370,7 +370,7 @@ void UUVEditorMode::ModeTick(float DeltaTime)
 				0, ToolSetupUtil::GetCustomTwoSidedDepthOffsetMaterial(
 					GetToolManager(),
 					(FLinearColor)TriangleColor,
-					2,  //depth offset
+					TriangleDepthOffset,
 					TriangleOpacity));
 		}
 		
diff --git a/Engine/Plugins/Experimental/UVEditor/Source/UVEditor/Public/UVEditorMode.h b/Engine/Plugins/Experimental/UVEditor/Source/UVEditor/Public/UVEditorMode.h
--- a/Engine/Plugins/Experimental/UVEditor/Source/UVEditor/Public/UVEditorMode.h
+++ b/Engine/Plugins/Experimental/UVEditor/Source/UVEditor/Public/UVEditorMode.h
@@ -98,6 +98,8 @@ protected:
 
 	// Wireframe Display Properties
 	float TriangleOpacity = 1.0;
+	const float TriangleDepthOffset = 0.5;
+	const float WireframeDepthOffset = 0.6;	
 	FColor TriangleColor = FColor(50, 194, 219);
 	FColor WireframeColor = FColor(50, 100, 219);
 	FColor IslandBorderColor = FColor(103, 52, 235);
diff --git a/Engine/Plugins/Experimental/UVEditor/Source/UVEditorTools/Private/UVSelectTool.cpp b/Engine/Plugins/Experimental/UVEditor/Source/UVEditorTools/Private/UVSelectTool.cpp
--- a/Engine/Plugins/Experimental/UVEditor/Source/UVEditorTools/Private/UVSelectTool.cpp
+++ b/Engine/Plugins/Experimental/UVEditor/Source/UVEditorTools/Private/UVSelectTool.cpp
@@ -695,7 +695,7 @@ void UUVSelectTool::UpdateLivePreviewLines()
 			LivePreviewLineSet->AddLine(
 				MeshTransform.TransformPosition(Vert1), 
 				MeshTransform.TransformPosition(Vert2), 
-				FColor::Yellow, 2, 1.5);
+				FColor::Yellow, LivePreviewHighlightThickness, LivePreviewHighlightDepthOffset);
 		}	
 	}
 }
diff --git a/Engine/Plugins/Experimental/UVEditor/Source/UVEditorTools/Public/UVSelectTool.h b/Engine/Plugins/Experimental/UVEditor/Source/UVEditorTools/Public/UVSelectTool.h
--- a/Engine/Plugins/Experimental/UVEditor/Source/UVEditorTools/Public/UVSelectTool.h
+++ b/Engine/Plugins/Experimental/UVEditor/Source/UVEditorTools/Public/UVSelectTool.h
@@ -210,6 +210,8 @@ protected:
 	ESelectToolAction PendingAction;
 
 	void ApplyAction(ESelectToolAction ActionType);
+	const float LivePreviewHighlightThickness = 2.0;
+	const float LivePreviewHighlightDepthOffset = 0.5;
 
 	// We need this flag so that SetGizmoVisibility can be called before Setup() by the tool builder.
 	bool bGizmoEnabled = false;
