UVEditor: Change triangle material response to background appearance out of critical section, fix gizmo visibility overtop of drawn UV triangles

#rnx

diff --git a/Engine/Plugins/Experimental/UVEditor/Source/UVEditor/Private/UVEditorMode.cpp b/Engine/Plugins/Experimental/UVEditor/Source/UVEditor/Private/UVEditorMode.cpp
--- a/Engine/Plugins/Experimental/UVEditor/Source/UVEditor/Private/UVEditorMode.cpp
+++ b/Engine/Plugins/Experimental/UVEditor/Source/UVEditor/Private/UVEditorMode.cpp
@@ -69,9 +69,11 @@ void UUVEditorMode::Enter()
 
 	BackgroundVisualization = NewObject<UUVEditorBackgroundPreview>(this);
 	BackgroundVisualization->CreateInWorld(GetWorld(), FTransform::Identity);
-
 	EditorModeToolkit->SetBackgroundSettings(BackgroundVisualization->Settings);
 
+	BackgroundVisualization->Settings->WatchProperty(BackgroundVisualization->Settings->bVisible, [this](bool IsVisible) { UpdateTriangleMaterialBasedOnBackground(IsVisible); });
+	PropertyObjectsToTick.Add(BackgroundVisualization->Settings);
+
 	RegisterTools();
 
 	ActivateDefaultTool();
@@ -313,6 +315,33 @@ void UUVEditorMode::ApplyChanges()
 	GetToolManager()->EndUndoTransaction();
 }
 
+void UUVEditorMode::UpdateTriangleMaterialBasedOnBackground(bool IsBackgroundVisible)
+{
+	// We adjust the mesh opacity depending on whether we're layered over the background or not.
+	if (IsBackgroundVisible)
+	{
+		TriangleOpacity = 0.25;
+		TriangleDepthOffset = .5;
+	}
+	else
+	{
+		TriangleOpacity = 1.0;
+		TriangleDepthOffset = -.1;
+	}
+
+	// Modify the material of the unwrapped mesh to account for the presence/absence of the background, 
+	// changing the opacity as set just above.
+	for (UUVEditorToolMeshInput* ToolInputObject : ToolInputObjects)
+	{
+		ToolInputObject->UnwrapPreview->PreviewMesh->SetMaterial(
+			0, ToolSetupUtil::GetCustomTwoSidedDepthOffsetMaterial(
+				GetToolManager(),
+				(FLinearColor)TriangleColor,
+				TriangleDepthOffset, 
+				TriangleOpacity));
+	}
+}
+
 void UUVEditorMode::ModeTick(float DeltaTime)
 {
 	Super::ModeTick(DeltaTime);
@@ -343,40 +372,7 @@ void UUVEditorMode::ModeTick(float DeltaTime)
 	if (BackgroundVisualization)
 	{
 		BackgroundVisualization->OnTick(DeltaTime);
-		if(BackgroundVisualization->Settings->IsPropertySetEnabled())
-		{
-			BackgroundVisualization->Settings->CheckAndUpdateWatched();
-		}
-		else
-		{
-			BackgroundVisualization->Settings->SilentUpdateWatched();
-		}
-
-        // We adjust the mesh opacity depending on whether we're layered over the background or not.
-		if (BackgroundVisualization->Settings->bVisible)
-		{
-			TriangleOpacity = 0.25;
-		}
-		else
-		{
-			TriangleOpacity = 1.0;
-		}
-
-        // Modify the material of the unwrapped mesh to account for the presence/absence of the background, 
-        // changing the opacity as set just above.
-		for (UUVEditorToolMeshInput* ToolInputObject : ToolInputObjects)
-		{
-			ToolInputObject->UnwrapPreview->PreviewMesh->SetMaterial(
-				0, ToolSetupUtil::GetCustomTwoSidedDepthOffsetMaterial(
-					GetToolManager(),
-					(FLinearColor)TriangleColor,
-					TriangleDepthOffset,
-					TriangleOpacity));
-		}
-		
 	}
-
-
 }
 
 #undef LOCTEXT_NAMESPACE
diff --git a/Engine/Plugins/Experimental/UVEditor/Source/UVEditor/Public/UVEditorMode.h b/Engine/Plugins/Experimental/UVEditor/Source/UVEditor/Public/UVEditorMode.h
--- a/Engine/Plugins/Experimental/UVEditor/Source/UVEditor/Public/UVEditorMode.h
+++ b/Engine/Plugins/Experimental/UVEditor/Source/UVEditor/Public/UVEditorMode.h
@@ -95,10 +95,11 @@ protected:
 	virtual void BindCommands() override;
 	virtual void OnToolStarted(UInteractiveToolManager* Manager, UInteractiveTool* Tool) {}
 	virtual void OnToolEnded(UInteractiveToolManager* Manager, UInteractiveTool* Tool) {}
+	void UpdateTriangleMaterialBasedOnBackground(bool IsBackgroundVisible);
 
 	// Wireframe Display Properties
 	float TriangleOpacity = 1.0;
-	const float TriangleDepthOffset = 0.5;
+	float TriangleDepthOffset = 0.5;
 	const float WireframeDepthOffset = 0.6;	
 	FColor TriangleColor = FColor(50, 194, 219);
 	FColor WireframeColor = FColor(50, 100, 219);
