UVEditor: Fix handling of PropertySets to support dynamic enabling/disabling within UV Editor

#rnx

diff --git a/Engine/Plugins/Experimental/UVEditor/Source/UVEditor/Private/UVEditorModeToolkit.cpp b/Engine/Plugins/Experimental/UVEditor/Source/UVEditor/Private/UVEditorModeToolkit.cpp
--- a/Engine/Plugins/Experimental/UVEditor/Source/UVEditor/Private/UVEditorModeToolkit.cpp
+++ b/Engine/Plugins/Experimental/UVEditor/Source/UVEditor/Private/UVEditorModeToolkit.cpp
@@ -323,10 +323,28 @@ void FUVEditorModeToolkit::SetBackgroundSettings(UObject* InSettingsObject)
 }
 
 
+void FUVEditorModeToolkit::UpdateActiveToolProperties()
+{
+	UInteractiveTool* CurTool = GetScriptableEditorMode()->GetToolManager()->GetActiveTool(EToolSide::Left);
+	if (CurTool != nullptr)
+	{
+		DetailsView->SetObjects(CurTool->GetToolProperties(true));
+	}
+}
+
+void FUVEditorModeToolkit::InvalidateCachedDetailPanelState(UObject* ChangedObject)
+{
+	DetailsView->InvalidateCachedState();
+}
+
 void FUVEditorModeToolkit::OnToolStarted(UInteractiveToolManager* Manager, UInteractiveTool* Tool)
 {
 	FModeToolkit::OnToolStarted(Manager, Tool);
 	
+	UInteractiveTool* CurTool = GetScriptableEditorMode()->GetToolManager()->GetActiveTool(EToolSide::Left);
+	CurTool->OnPropertySetsModified.AddSP(this, &FUVEditorModeToolkit::UpdateActiveToolProperties);
+	CurTool->OnPropertyModifiedDirectlyByTool.AddSP(this, &FUVEditorModeToolkit::InvalidateCachedDetailPanelState);
+
 	ActiveToolName = Tool->GetToolInfo().ToolDisplayName;
 
 	if (GetScriptableEditorMode()->GetInteractiveToolsContext()->ActiveToolHasAccept())
@@ -355,6 +373,13 @@ void FUVEditorModeToolkit::OnToolEnded(UInteractiveToolManager* Manager, UIntera
 	{
 		GetToolkitHost()->RemoveViewportOverlayWidget(ViewportOverlayWidget.ToSharedRef());
 	}
+
+	UInteractiveTool* CurTool = GetScriptableEditorMode()->GetToolManager()->GetActiveTool(EToolSide::Left);
+	if (CurTool)
+	{
+		CurTool->OnPropertySetsModified.RemoveAll(this);
+		CurTool->OnPropertyModifiedDirectlyByTool.RemoveAll(this);
+	}
 }
 
 #undef LOCTEXT_NAMESPACE
diff --git a/Engine/Plugins/Experimental/UVEditor/Source/UVEditor/Public/UVEditorModeToolkit.h b/Engine/Plugins/Experimental/UVEditor/Source/UVEditor/Public/UVEditorModeToolkit.h
--- a/Engine/Plugins/Experimental/UVEditor/Source/UVEditor/Public/UVEditorModeToolkit.h
+++ b/Engine/Plugins/Experimental/UVEditor/Source/UVEditor/Public/UVEditorModeToolkit.h
@@ -63,4 +63,8 @@ protected:
 	TSharedPtr<SWidget> ViewportOverlayWidget;
 
 	FText ActiveToolName;
+
+	void UpdateActiveToolProperties();
+	void InvalidateCachedDetailPanelState(UObject* ChangedObject);
+
 };
\ No newline at end of file
