UVEditor: Fix overlay copy incorrectly dealing with parent vert IDs in some remeshing situations.

#rnx

diff --git a/Engine/Plugins/Experimental/UVEditor/Source/UVEditorTools/Private/ToolTargets/UVEditorToolMeshInput.cpp b/Engine/Plugins/Experimental/UVEditor/Source/UVEditorTools/Private/ToolTargets/UVEditorToolMeshInput.cpp
--- a/Engine/Plugins/Experimental/UVEditor/Source/UVEditorTools/Private/ToolTargets/UVEditorToolMeshInput.cpp
+++ b/Engine/Plugins/Experimental/UVEditor/Source/UVEditorTools/Private/ToolTargets/UVEditorToolMeshInput.cpp
@@ -63,11 +63,24 @@ namespace UVEditorToolMeshInputLocals
 		}
 	}
 
-	// Aside from potentially  doing partial updates based on ChangedElements and ChangedConnectivityTids, this
-	// function also differs from a simple overlay copy in that it does not copy the parent pointer arrays.
-	// Also it assumes the existence of all the Tids.
-	void CopyMeshOverlay(const UE::Geometry::FDynamicMeshUVOverlay& OverlayIn, UE::Geometry::FDynamicMeshUVOverlay& OverlayOut,
-		const TArray<int32>* ChangedElements, const TArray<int32>* ChangedConnectivityTids)
+	/**
+	 * Copy all or parts of a mesh overlay into another mesh overlay.
+	 * 
+	 * WARNING: except in narrow cases where we are able to revert to a simple Copy() call, the function
+	 * forcibly sets the parent vertices of any elements touched by a SetTriangle() call made in the function.
+	 * This is necessary, for example, if you have a two triangle mesh with two UV islands and you want to
+	 * swap the assignment of the islands. However it means that the function could silently place the overlay 
+	 * in an invalid state if, for instance, ChangedConnectivityTids are not complete (in the previous example,
+	 * imagine only one of the triangles is included in ChangedConnectivityTids).
+	 * Consider using a checkSlow to check the validity of the output overlay to make sure you did not miss something. 
+	 *
+	 * @param bMeshesHaveSameTopology If true, underlying meshes are topologically identical, so we can use a
+	 *   simple copy when we are not constraining the copied elements/triangles. Can be set true when copying
+	 *   between a canonical/preview version of the same mesh, but must be false when copying between an unwrap
+	 *   mesh and an applied mesh.
+	 */
+	void CopyMeshOverlay(const UE::Geometry::FDynamicMeshUVOverlay& OverlayIn, UE::Geometry::FDynamicMeshUVOverlay& OverlayOut, 
+		bool bMeshesHaveSameTopology, const TArray<int32>* ChangedElements, const TArray<int32>* ChangedConnectivityTids)
 	{
 		auto UpdateOrInsertElement = [&OverlayIn, &OverlayOut](int32 ElementID) {
 			if (OverlayOut.IsElement(ElementID))
@@ -81,17 +94,51 @@ namespace UVEditorToolMeshInputLocals
 			}
 		};
 
+		// @param PotentiallyFreedElementsOut If not null, any elements that used to be referenced by
+		// the changed triangle are placed here so they can be checked for freeing later.
+		auto ResetTriangleWithParenting = [&OverlayIn, &OverlayOut](int32 Tid, TSet<int32>* PotentiallyFreedElementsOut) {
+			if (PotentiallyFreedElementsOut)
+			{
+				FIndex3i OldElementTri = OverlayIn.GetTriangle(Tid);
+				for (int i = 0; i < 3; ++i)
+				{
+					PotentiallyFreedElementsOut->Add(OldElementTri[i]);
+				}
+			}
+
+			// Force reset the parent pointers if necessary
+			FIndex3i NewElementTri = OverlayIn.GetTriangle(Tid);
+			FIndex3i ParentTriInOutput = OverlayOut.GetParentMesh()->GetTriangle(Tid);
+			for (int32 i = 0; i < 3; ++i)
+			{
+				if (OverlayOut.GetParentVertex(NewElementTri[i]) != ParentTriInOutput[i])
+				{
+					OverlayOut.SetParentVertex(NewElementTri[i], ParentTriInOutput[i]);
+				}
+			}
+
+			// Now set the triangle. Don't free elements since other new triangles might use them.
+			OverlayOut.SetTriangle(Tid, NewElementTri, false);
+		};
+
 		if (!ChangedElements && !ChangedConnectivityTids)
 		{
-			for (int32 ElementID : OverlayIn.ElementIndicesItr())
+			if (bMeshesHaveSameTopology)
 			{
-				UpdateOrInsertElement(ElementID);
+				OverlayOut.Copy(OverlayIn);
 			}
-			for (int32 Tid : OverlayIn.GetParentMesh()->TriangleIndicesItr())
+			else
 			{
-				OverlayOut.SetTriangle(Tid, OverlayIn.GetTriangle(Tid));
+				for (int32 ElementID : OverlayIn.ElementIndicesItr())
+				{
+					UpdateOrInsertElement(ElementID);
+				}
+				for (int32 Tid : OverlayIn.GetParentMesh()->TriangleIndicesItr())
+				{
+					ResetTriangleWithParenting(Tid, nullptr);
+				}
+				OverlayOut.FreeUnusedElements();
 			}
-
 			return;
 		}
 
@@ -104,10 +151,13 @@ namespace UVEditorToolMeshInputLocals
 		}
 		if (ChangedConnectivityTids)
 		{
+			TSet<int32> PotentiallyFreedElements;
 			for (int32 Tid : *ChangedConnectivityTids)
 			{
-				OverlayOut.SetTriangle(Tid, OverlayIn.GetTriangle(Tid));
+				ResetTriangleWithParenting(Tid, &PotentiallyFreedElements);
 			}
+
+			OverlayOut.FreeUnusedElements(&PotentiallyFreedElements);
 		}
 	}
 }
@@ -241,7 +291,7 @@ void UUVEditorToolMeshInput::UpdateAppliedPreviewFromUnwrapPreview(const TArray<
 			const FDynamicMeshUVOverlay* SourceOverlay = SourceUnwrapMesh->Attributes()->PrimaryUV();
 			FDynamicMeshUVOverlay* DestOverlay = Mesh.Attributes()->PrimaryUV();
 
-			CopyMeshOverlay(*SourceOverlay, *DestOverlay, ChangedVids, ChangedConnectivityTids);
+			CopyMeshOverlay(*SourceOverlay, *DestOverlay, false, ChangedVids, ChangedConnectivityTids);
 	}, false);
 
 	if (FastRenderUpdateTids) {
@@ -267,7 +317,7 @@ void UUVEditorToolMeshInput::UpdateUnwrapPreviewFromAppliedPreview(const TArray<
 
 		// Also copy the actual overlay
 		FDynamicMeshUVOverlay* DestOverlay = Mesh.Attributes()->PrimaryUV();
-		CopyMeshOverlay(*SourceOverlay, *DestOverlay, ChangedElementIDs, ChangedConnectivityTids);
+		CopyMeshOverlay(*SourceOverlay, *DestOverlay, false, ChangedElementIDs, ChangedConnectivityTids);
 	}, false);
 
 	if (FastRenderUpdateTids) {
@@ -296,7 +346,7 @@ void UUVEditorToolMeshInput::UpdateCanonicalFromPreviews(const TArray<int32>* Ch
 	// Update the overlay in AppliedCanonical from overlay in AppliedPreview
 	const FDynamicMeshUVOverlay* SourceOverlay = AppliedPreview->PreviewMesh->GetMesh()->Attributes()->GetUVLayer(UVLayerIndex);
 	FDynamicMeshUVOverlay* DestOverlay = AppliedCanonical->Attributes()->GetUVLayer(UVLayerIndex);
-	CopyMeshOverlay(*SourceOverlay, *DestOverlay, ChangedVids, ChangedConnectivityTids);
+	CopyMeshOverlay(*SourceOverlay, *DestOverlay, true, ChangedVids, ChangedConnectivityTids);
 }
 
 void UUVEditorToolMeshInput::UpdatePreviewsFromCanonical(const TArray<int32>* ChangedVids, const TArray<int32>* ChangedConnectivityTids, const TArray<int32>* FastRenderUpdateTids)
@@ -328,7 +378,7 @@ void UUVEditorToolMeshInput::UpdatePreviewsFromCanonical(const TArray<int32>* Ch
 	{
 		const FDynamicMeshUVOverlay* SourceOverlay = AppliedCanonical->Attributes()->GetUVLayer(UVLayerIndex);
 		FDynamicMeshUVOverlay* DestOverlay = Mesh.Attributes()->GetUVLayer(UVLayerIndex);
-		CopyMeshOverlay(*SourceOverlay, *DestOverlay, ChangedVids, ChangedConnectivityTids);
+		CopyMeshOverlay(*SourceOverlay, *DestOverlay, true, ChangedVids, ChangedConnectivityTids);
 	}, false);
 	if (FastRenderUpdateTids) {
 		AppliedPreview->PreviewMesh->NotifyRegionDeferredEditCompleted(*FastRenderUpdateTids,
@@ -355,7 +405,7 @@ void UUVEditorToolMeshInput::UpdateAllFromUnwrapCanonical(
 	// Update AppliedCanonical
 	FDynamicMeshUVOverlay* SourceOverlay = UnwrapCanonical->Attributes()->PrimaryUV();
 	FDynamicMeshUVOverlay* DestOverlay = AppliedCanonical->Attributes()->GetUVLayer(UVLayerIndex);
-	CopyMeshOverlay(*SourceOverlay, *DestOverlay, ChangedVids, ChangedConnectivityTids);
+	CopyMeshOverlay(*SourceOverlay, *DestOverlay, false, ChangedVids, ChangedConnectivityTids);
 
 	UpdatePreviewsFromCanonical(ChangedVids, ChangedConnectivityTids, FastRenderUpdateTids);
 }
@@ -385,7 +435,7 @@ void UUVEditorToolMeshInput::UpdateOtherUnwrap(const FDynamicMesh3& SourceUnwrap
 
 	const FDynamicMeshUVOverlay* SourceOverlay = SourceUnwrapMesh.Attributes()->PrimaryUV();
 	FDynamicMeshUVOverlay* DestOverlay = DestUnwrapMesh.Attributes()->PrimaryUV();
-	CopyMeshOverlay(*SourceOverlay, *DestOverlay, ChangedVids, ChangedConnectivityTids);
+	CopyMeshOverlay(*SourceOverlay, *DestOverlay, true, ChangedVids, ChangedConnectivityTids);
 }
 
 void UUVEditorToolMeshInput::UpdateFromCanonicalUnwrapUsingMeshChange(const FDynamicMeshChange& UnwrapCanonicalMeshChange)
