UVEditor: Add support for the UV parameterization tool into the editor

#rnx

diff --git a/Engine/Plugins/Experimental/MeshModelingToolsetExp/Source/MeshModelingToolsEditorOnly/Private/ParameterizeMeshTool.cpp b/Engine/Plugins/Experimental/MeshModelingToolsetExp/Source/MeshModelingToolsEditorOnly/Private/ParameterizeMeshTool.cpp
--- a/Engine/Plugins/Experimental/MeshModelingToolsetExp/Source/MeshModelingToolsEditorOnly/Private/ParameterizeMeshTool.cpp
+++ b/Engine/Plugins/Experimental/MeshModelingToolsetExp/Source/MeshModelingToolsEditorOnly/Private/ParameterizeMeshTool.cpp
@@ -8,6 +8,7 @@
 #include "DynamicMesh/DynamicMesh3.h"
 #include "FaceGroupUtil.h"
 #include "ParameterizationOps/ParameterizeMeshOp.h"
+#include "Properties/ParameterizeMeshProperties.h"
 
 #include "ExplicitUseGeometryMathTypes.h"		// using UE::Geometry::(math types)
 using namespace UE::Geometry;
diff --git a/Engine/Plugins/Experimental/MeshModelingToolsetExp/Source/MeshModelingToolsEditorOnly/Public/ParameterizeMeshTool.h b/Engine/Plugins/Experimental/MeshModelingToolsetExp/Source/MeshModelingToolsEditorOnly/Public/ParameterizeMeshTool.h
--- a/Engine/Plugins/Experimental/MeshModelingToolsetExp/Source/MeshModelingToolsEditorOnly/Public/ParameterizeMeshTool.h
+++ b/Engine/Plugins/Experimental/MeshModelingToolsetExp/Source/MeshModelingToolsEditorOnly/Public/ParameterizeMeshTool.h
@@ -16,7 +16,10 @@
 class UDynamicMeshComponent;
 class UMaterialInterface;
 class UMaterialInstanceDynamic;
-
+class UParameterizeMeshToolProperties;
+class UParameterizeMeshToolUVAtlasProperties;
+class UParameterizeMeshToolXAtlasProperties;
+class UParameterizeMeshToolPatchBuilderProperties;
 
 UCLASS()
 class MESHMODELINGTOOLSEDITORONLY_API UParameterizeMeshToolBuilder : public USingleSelectionMeshEditingToolBuilder
@@ -27,104 +30,6 @@ public:
 };
 
 
-UENUM()
-enum class EParameterizeMeshUVMethod
-{
-	// keep values the same as UE::Geometry::EParamOpBackend!
-	/** Compute Automatic UVs using PatchBuilder */
-	PatchBuilder = 0,
-	/** Compute Automatic UVs using XAtlas */
-	XAtlas = 1,
-	/** Compute Automatic UVs using UVAtlas */
-	UVAtlas = 2,
-};
-
-
-
-UCLASS()
-class MESHMODELINGTOOLSEDITORONLY_API UParameterizeMeshToolProperties : public UInteractiveToolPropertySet
-{
-	GENERATED_BODY()
-public:
-	/** Automatic UV Generation technique to use */
-	UPROPERTY(EditAnywhere, Category = Options)
-	EParameterizeMeshUVMethod Method = EParameterizeMeshUVMethod::PatchBuilder;
-};
-
-
-
-/**
- * Settings for the UVAtlas Automatic UV Generation Method
- */
-UCLASS()
-class MESHMODELINGTOOLSEDITORONLY_API UParameterizeMeshToolUVAtlasProperties : public UInteractiveToolPropertySet
-{
-	GENERATED_BODY()
-public:
-	/** Maximum amount of stretch, from none to any.  If zero stretch is specified each triangle will likely be its own chart */
-	UPROPERTY(EditAnywhere, Category = UVAtlas,  meta = (UIMin = "0", UIMax = "1", ClampMin = "0", ClampMax = "1"))
-	float ChartStretch = 0.11f;
-
-	/** Hint at number of Charts. 0 (Default) means UVAtlas will decide */
-	UPROPERTY(EditAnywhere, Category = UVAtlas, meta = (UIMin = "0", UIMax = "100", ClampMin = "0", ClampMax = "10000"))
-	int NumCharts = 0;
-};
-
-
-UCLASS()
-class MESHMODELINGTOOLSEDITORONLY_API UParameterizeMeshToolXAtlasProperties : public UInteractiveToolPropertySet
-{
-	GENERATED_BODY()
-public:
-	/** Number of solve iterations. Higher values generally result in better charts. */
-	UPROPERTY(EditAnywhere, Category = XAtlas, meta = (UIMin = "1", UIMax = "10", ClampMin = "1", ClampMax = "1000"))
-	int MaxIterations = 1;
-};
-
-
-
-UCLASS()
-class MESHMODELINGTOOLSEDITORONLY_API UParameterizeMeshToolPatchBuilderProperties : public UInteractiveToolPropertySet
-{
-	GENERATED_BODY()
-public:
-	/** Number of initial patches mesh will be split into before computing island merging */
-	UPROPERTY(EditAnywhere, Category = "PatchBuilder", meta = (UIMin = "1", UIMax = "1000", ClampMin = "1", ClampMax = "99999999"))
-	int InitialPatches = 100;
-
-	/** This parameter controls alignment of the initial patches to creases in the mesh */
-	UPROPERTY(EditAnywhere, Category = "PatchBuilder", meta = (UIMin = "0.1", UIMax = "2.0", ClampMin = "0.01", ClampMax = "100.0"))
-	float CurvatureAlignment = 1.0f;
-
-	/** Distortion/Stretching Threshold for island merging - larger values increase the allowable UV stretching */
-	UPROPERTY(EditAnywhere, Category = "PatchBuilder", meta = (UIMin = "1.0", UIMax = "5.0", ClampMin = "1.0", EditCondition = "bIslandMerging == true"))
-	float MergingThreshold = 1.5f;
-
-	/** UV islands will not be merged if their average face normals deviate by larger than this amount */
-	UPROPERTY(EditAnywhere, Category = "PatchBuilder", meta = (UIMin = "0.0", UIMax = "90.0", ClampMin = "0.0", ClampMax = "180.0"))
-	float MaxAngleDeviation = 45.0f;
-
-	/** Number of smoothing steps to apply in the ExpMap UV Generation method (Smoothing slightly increases distortion but produces more stable results) */
-	UPROPERTY(EditAnywhere, Category = "PatchBuilder", meta = (UIMin = "0", UIMax = "25", ClampMin = "0", ClampMax = "1000"))
-	int SmoothingSteps = 5;
-
-	/** Smoothing parameter, larger values result in faster smoothing in each step */
-	UPROPERTY(EditAnywhere, Category = "PatchBuilder", meta = (UIMin = "0", UIMax = "1.0", ClampMin = "0", ClampMax = "1.0"))
-	float SmoothingAlpha = 0.25f;
-
-
-	/** If enabled, result UVs are automatically packed into the standard UV 0-1 square */
-	UPROPERTY(EditAnywhere, Category = "PatchBuilder")
-	bool bAutoPack = true;
-
-	/** Target texture resolution used for UV packing, which determines gutter size */
-	UPROPERTY(EditAnywhere, Category = "PatchBuilder", meta = (UIMin = "64", UIMax = "2048", ClampMin = "2", ClampMax = "4096", EditCondition = bAutoPack))
-	int TextureResolution = 1024;
-};
-
-
-
-
 /**
  * UParameterizeMeshTool automatically decomposes the input mesh into charts, solves for UVs,
  * and then packs the resulting charts
diff --git a/Engine/Plugins/Experimental/MeshModelingToolsetExp/Source/ModelingOperatorsEditorOnly/Private/ParameterizationOps/ParameterizeMeshOp.cpp b/Engine/Plugins/Experimental/MeshModelingToolsetExp/Source/ModelingOperatorsEditorOnly/Private/ParameterizationOps/ParameterizeMeshOp.cpp
--- a/Engine/Plugins/Experimental/MeshModelingToolsetExp/Source/ModelingOperatorsEditorOnly/Private/ParameterizationOps/ParameterizeMeshOp.cpp
+++ b/Engine/Plugins/Experimental/MeshModelingToolsetExp/Source/ModelingOperatorsEditorOnly/Private/ParameterizationOps/ParameterizeMeshOp.cpp
@@ -1,6 +1,7 @@
 // Copyright Epic Games, Inc. All Rights Reserved.
 
 #include "ParameterizationOps/ParameterizeMeshOp.h"
+#include "Properties/ParameterizeMeshProperties.h"
 
 #include "DynamicMesh/DynamicMeshAttributeSet.h"
 #include "Selections/MeshConnectedComponents.h"
@@ -436,4 +437,45 @@ void FParameterizeMeshOp::CalculateResult(FProgressCancel* Progress)
 }
 
 
+
+TUniquePtr<FDynamicMeshOperator> UParameterizeMeshOperatorFactory::MakeNewOperator()
+{
+	TUniquePtr<FParameterizeMeshOp> Op = MakeUnique<FParameterizeMeshOp>();
+
+	Op->InputMesh = OriginalMesh;
+
+	switch (Settings->Method)
+	{
+	case EParameterizeMeshUVMethod::PatchBuilder:
+		Op->Method = EParamOpBackend::PatchBuilder;
+		Op->InitialPatchCount = PatchBuilderProperties->InitialPatches;
+		Op->PatchCurvatureAlignmentWeight = PatchBuilderProperties->CurvatureAlignment;
+		Op->PatchMergingMetricThresh = PatchBuilderProperties->MergingThreshold;
+		Op->PatchMergingAngleThresh = PatchBuilderProperties->MaxAngleDeviation;
+		Op->ExpMapNormalSmoothingSteps = PatchBuilderProperties->SmoothingSteps;
+		Op->ExpMapNormalSmoothingAlpha = PatchBuilderProperties->SmoothingAlpha;
+		Op->bEnablePacking = PatchBuilderProperties->bAutoPack;
+		Op->Width = PatchBuilderProperties->TextureResolution;
+		Op->Height = PatchBuilderProperties->TextureResolution;
+		break;
+	case EParameterizeMeshUVMethod::UVAtlas:
+		Op->Method = EParamOpBackend::UVAtlas;		
+		Op->Stretch = UVAtlasProperties->ChartStretch;
+		Op->NumCharts = UVAtlasProperties->NumCharts;
+		break;
+	case EParameterizeMeshUVMethod::XAtlas:
+		Op->Method = EParamOpBackend::XAtlas;
+		Op->XAtlasMaxIterations = XAtlasProperties->MaxIterations;
+		break;
+	}
+
+	Op->UVLayer = GetSelectedUVChannel();	
+
+	UE::Geometry::FTransform3d Transform(TargetTransform);
+	Op->SetTransform(Transform);
+
+	return Op;
+}
+
+
 #undef LOCTEXT_NAMESPACE
\ No newline at end of file
diff --git a/Engine/Plugins/Experimental/MeshModelingToolsetExp/Source/ModelingOperatorsEditorOnly/Public/ParameterizationOps/ParameterizeMeshOp.h b/Engine/Plugins/Experimental/MeshModelingToolsetExp/Source/ModelingOperatorsEditorOnly/Public/ParameterizationOps/ParameterizeMeshOp.h
--- a/Engine/Plugins/Experimental/MeshModelingToolsetExp/Source/ModelingOperatorsEditorOnly/Public/ParameterizationOps/ParameterizeMeshOp.h
+++ b/Engine/Plugins/Experimental/MeshModelingToolsetExp/Source/ModelingOperatorsEditorOnly/Public/ParameterizationOps/ParameterizeMeshOp.h
@@ -7,6 +7,12 @@
 #include "ModelingOperators.h"
 #include "DynamicMesh/DynamicMeshAttributeSet.h"
 
+#include "ParameterizeMeshOp.generated.h"
+
+class UParameterizeMeshToolProperties;
+class UParameterizeMeshToolUVAtlasProperties;
+class UParameterizeMeshToolXAtlasProperties;
+class UParameterizeMeshToolPatchBuilderProperties;
 
 namespace UE
 {
@@ -109,3 +115,32 @@ protected:
 
 } // end namespace UE::Geometry
 } // end namespace UE
+
+/**
+ * Can be hooked up to a UMeshOpPreviewWithBackgroundCompute to perform UV parameterization operations.
+ */
+UCLASS()
+class MODELINGOPERATORSEDITORONLY_API UParameterizeMeshOperatorFactory : public UObject, public UE::Geometry::IDynamicMeshOperatorFactory
+{
+	GENERATED_BODY()
+
+public:
+	// IDynamicMeshOperatorFactory API
+	virtual TUniquePtr<UE::Geometry::FDynamicMeshOperator> MakeNewOperator() override;
+
+	UPROPERTY()
+	TObjectPtr<UParameterizeMeshToolProperties> Settings = nullptr;
+
+	UPROPERTY()
+	TObjectPtr<UParameterizeMeshToolUVAtlasProperties> UVAtlasProperties = nullptr;
+
+	UPROPERTY()
+	TObjectPtr<UParameterizeMeshToolXAtlasProperties> XAtlasProperties = nullptr;
+
+	UPROPERTY()
+	TObjectPtr<UParameterizeMeshToolPatchBuilderProperties> PatchBuilderProperties = nullptr;
+
+	TSharedPtr<UE::Geometry::FDynamicMesh3, ESPMode::ThreadSafe> OriginalMesh;
+	TUniqueFunction<int32()> GetSelectedUVChannel = []() { return 0; };
+	FTransform TargetTransform;
+};
\ No newline at end of file
diff --git a/Engine/Plugins/Experimental/MeshModelingToolsetExp/Source/ModelingOperatorsEditorOnly/Public/Properties/ParameterizeMeshProperties.h b/Engine/Plugins/Experimental/MeshModelingToolsetExp/Source/ModelingOperatorsEditorOnly/Public/Properties/ParameterizeMeshProperties.h
new file mode 100644
--- /dev/null
+++ b/Engine/Plugins/Experimental/MeshModelingToolsetExp/Source/ModelingOperatorsEditorOnly/Public/Properties/ParameterizeMeshProperties.h
@@ -0,0 +1,101 @@
+// Copyright Epic Games, Inc. All Rights Reserved.
+
+#pragma once
+
+#include "CoreMinimal.h"
+
+#include "InteractiveTool.h"
+
+#include "ParameterizeMeshProperties.generated.h"
+
+UENUM()
+enum class EParameterizeMeshUVMethod
+{
+	// keep values the same as UE::Geometry::EParamOpBackend!
+	/** Compute Automatic UVs using PatchBuilder */
+	PatchBuilder = 0,
+	/** Compute Automatic UVs using XAtlas */
+	XAtlas = 1,
+	/** Compute Automatic UVs using UVAtlas */
+	UVAtlas = 2,
+};
+
+
+UCLASS()
+class MODELINGOPERATORSEDITORONLY_API UParameterizeMeshToolProperties : public UInteractiveToolPropertySet
+{
+	GENERATED_BODY()
+public:
+	/** Automatic UV Generation technique to use */
+	UPROPERTY(EditAnywhere, Category = Options)
+	EParameterizeMeshUVMethod Method = EParameterizeMeshUVMethod::PatchBuilder;
+};
+
+
+/**
+ * Settings for the UVAtlas Automatic UV Generation Method
+ */
+UCLASS()
+class MODELINGOPERATORSEDITORONLY_API UParameterizeMeshToolUVAtlasProperties : public UInteractiveToolPropertySet
+{
+	GENERATED_BODY()
+public:
+	/** Maximum amount of stretch, from none to any.  If zero stretch is specified each triangle will likely be its own chart */
+	UPROPERTY(EditAnywhere, Category = UVAtlas, meta = (UIMin = "0", UIMax = "1", ClampMin = "0", ClampMax = "1"))
+	float ChartStretch = 0.11f;
+
+	/** Hint at number of Charts. 0 (Default) means UVAtlas will decide */
+	UPROPERTY(EditAnywhere, Category = UVAtlas, meta = (UIMin = "0", UIMax = "100", ClampMin = "0", ClampMax = "10000"))
+	int NumCharts = 0;
+};
+
+
+UCLASS()
+class MODELINGOPERATORSEDITORONLY_API UParameterizeMeshToolXAtlasProperties : public UInteractiveToolPropertySet
+{
+	GENERATED_BODY()
+public:
+	/** Number of solve iterations. Higher values generally result in better charts. */
+	UPROPERTY(EditAnywhere, Category = XAtlas, meta = (UIMin = "1", UIMax = "10", ClampMin = "1", ClampMax = "1000"))
+	int MaxIterations = 1;
+};
+
+
+UCLASS()
+class MODELINGOPERATORSEDITORONLY_API UParameterizeMeshToolPatchBuilderProperties : public UInteractiveToolPropertySet
+{
+	GENERATED_BODY()
+public:
+	/** Number of initial patches mesh will be split into before computing island merging */
+	UPROPERTY(EditAnywhere, Category = "PatchBuilder", meta = (UIMin = "1", UIMax = "1000", ClampMin = "1", ClampMax = "99999999"))
+	int InitialPatches = 100;
+
+	/** This parameter controls alignment of the initial patches to creases in the mesh */
+	UPROPERTY(EditAnywhere, Category = "PatchBuilder", meta = (UIMin = "0.1", UIMax = "2.0", ClampMin = "0.01", ClampMax = "100.0"))
+	float CurvatureAlignment = 1.0f;
+
+	/** Distortion/Stretching Threshold for island merging - larger values increase the allowable UV stretching */
+	UPROPERTY(EditAnywhere, Category = "PatchBuilder", meta = (UIMin = "1.0", UIMax = "5.0", ClampMin = "1.0", EditCondition = "bIslandMerging == true"))
+	float MergingThreshold = 1.5f;
+
+	/** UV islands will not be merged if their average face normals deviate by larger than this amount */
+	UPROPERTY(EditAnywhere, Category = "PatchBuilder", meta = (UIMin = "0.0", UIMax = "90.0", ClampMin = "0.0", ClampMax = "180.0"))
+	float MaxAngleDeviation = 45.0f;
+
+	/** Number of smoothing steps to apply in the ExpMap UV Generation method (Smoothing slightly increases distortion but produces more stable results) */
+	UPROPERTY(EditAnywhere, Category = "PatchBuilder", meta = (UIMin = "0", UIMax = "25", ClampMin = "0", ClampMax = "1000"))
+	int SmoothingSteps = 5;
+
+	/** Smoothing parameter, larger values result in faster smoothing in each step */
+	UPROPERTY(EditAnywhere, Category = "PatchBuilder", meta = (UIMin = "0", UIMax = "1.0", ClampMin = "0", ClampMax = "1.0"))
+	float SmoothingAlpha = 0.25f;
+
+
+	/** If enabled, result UVs are automatically packed into the standard UV 0-1 square */
+	UPROPERTY(EditAnywhere, Category = "PatchBuilder")
+	bool bAutoPack = true;
+
+	/** Target texture resolution used for UV packing, which determines gutter size */
+	UPROPERTY(EditAnywhere, Category = "PatchBuilder", meta = (UIMin = "64", UIMax = "2048", ClampMin = "2", ClampMax = "4096", EditCondition = bAutoPack))
+	int TextureResolution = 1024;
+};
\ No newline at end of file
diff --git a/Engine/Plugins/Experimental/UVEditor/Source/UVEditor/Private/UVEditorCommands.cpp b/Engine/Plugins/Experimental/UVEditor/Source/UVEditor/Private/UVEditorCommands.cpp
--- a/Engine/Plugins/Experimental/UVEditor/Source/UVEditor/Private/UVEditorCommands.cpp
+++ b/Engine/Plugins/Experimental/UVEditor/Source/UVEditor/Private/UVEditorCommands.cpp
@@ -26,6 +26,7 @@ void FUVEditorCommands::RegisterCommands()
 	UI_COMMAND(BeginSelectTool, "Select", "Switch to select tool", EUserInterfaceActionType::ToggleButton, FInputChord());
 	UI_COMMAND(BeginTransformTool, "Transform", "Switch to transform tool", EUserInterfaceActionType::ToggleButton, FInputChord());
 	UI_COMMAND(BeginLayoutTool, "Layout", "Switch to layout tool", EUserInterfaceActionType::ToggleButton, FInputChord());
+	UI_COMMAND(BeginParameterizeMeshTool, "Auto UV", "Switch to Auto UV tool", EUserInterfaceActionType::ToggleButton, FInputChord());
 
 	// These allow us to link up to pressed keys
 	UI_COMMAND(AcceptActiveTool, "Accept", "Accept the active tool", EUserInterfaceActionType::Button, FInputChord(EKeys::Enter));
diff --git a/Engine/Plugins/Experimental/UVEditor/Source/UVEditor/Private/UVEditorMode.cpp b/Engine/Plugins/Experimental/UVEditor/Source/UVEditor/Private/UVEditorMode.cpp
--- a/Engine/Plugins/Experimental/UVEditor/Source/UVEditor/Private/UVEditorMode.cpp
+++ b/Engine/Plugins/Experimental/UVEditor/Source/UVEditor/Private/UVEditorMode.cpp
@@ -20,6 +20,7 @@
 #include "ToolTargets/UVEditorToolMeshInput.h"
 #include "UVEditorCommands.h"
 #include "UVEditorLayoutTool.h"
+#include "UVEditorParameterizeMeshTool.h"
 #include "UVSelectTool.h"
 #include "UVEditorModeToolkit.h"
 #include "UVEditorSubsystem.h"
@@ -96,6 +97,10 @@ void UUVEditorMode::RegisterTools()
 	UUVEditorLayoutToolBuilder* UVEditorLayoutToolBuilder = NewObject<UUVEditorLayoutToolBuilder>();
 	UVEditorLayoutToolBuilder->Targets = &ToolInputObjects;
 	RegisterTool(CommandInfos.BeginLayoutTool, TEXT("UVLayoutTool"), UVEditorLayoutToolBuilder);
+
+	UUVEditorParameterizeMeshToolBuilder* UVEditorParameterizeMeshToolBuilder = NewObject<UUVEditorParameterizeMeshToolBuilder>();
+	UVEditorParameterizeMeshToolBuilder->Targets = &ToolInputObjects;
+	RegisterTool(CommandInfos.BeginParameterizeMeshTool, TEXT("UVParameterizeMeshTool"), UVEditorParameterizeMeshToolBuilder);
 }
 
 void UUVEditorMode::CreateToolkit()
@@ -373,6 +378,12 @@ void UUVEditorMode::ModeTick(float DeltaTime)
 	{
 		BackgroundVisualization->OnTick(DeltaTime);
 	}
+
+	for (int i = 0; i < ToolInputObjects.Num(); ++i)
+	{
+		ToolInputObjects[i]->AppliedPreview->Tick(DeltaTime);
+		ToolInputObjects[i]->UnwrapPreview->Tick(DeltaTime);
+	}
 }
 
 #undef LOCTEXT_NAMESPACE
diff --git a/Engine/Plugins/Experimental/UVEditor/Source/UVEditor/Private/UVEditorModeToolkit.cpp b/Engine/Plugins/Experimental/UVEditor/Source/UVEditor/Private/UVEditorModeToolkit.cpp
--- a/Engine/Plugins/Experimental/UVEditor/Source/UVEditor/Private/UVEditorModeToolkit.cpp
+++ b/Engine/Plugins/Experimental/UVEditor/Source/UVEditor/Private/UVEditorModeToolkit.cpp
@@ -185,6 +185,7 @@ void FUVEditorModeToolkit::Init(const TSharedPtr<IToolkitHost>& InitToolkitHost,
 	ToolbarBuilder.AddToolBarButton(CommandInfos.BeginSelectTool);
 	ToolbarBuilder.AddToolBarButton(CommandInfos.BeginTransformTool);
 	ToolbarBuilder.AddToolBarButton(CommandInfos.BeginLayoutTool);
+	ToolbarBuilder.AddToolBarButton(CommandInfos.BeginParameterizeMeshTool);
 
 	// Hook in the tool palette
 	ToolButtonsContainer->SetContent(ToolbarBuilder.MakeWidget());
diff --git a/Engine/Plugins/Experimental/UVEditor/Source/UVEditor/Public/UVEditorCommands.h b/Engine/Plugins/Experimental/UVEditor/Source/UVEditor/Public/UVEditorCommands.h
--- a/Engine/Plugins/Experimental/UVEditor/Source/UVEditor/Public/UVEditorCommands.h
+++ b/Engine/Plugins/Experimental/UVEditor/Source/UVEditor/Public/UVEditorCommands.h
@@ -21,6 +21,7 @@ public:
 	TSharedPtr<FUICommandInfo> BeginSelectTool;
 	TSharedPtr<FUICommandInfo> BeginTransformTool;
 	TSharedPtr<FUICommandInfo> BeginLayoutTool;
+	TSharedPtr<FUICommandInfo> BeginParameterizeMeshTool;
 
 	TSharedPtr<FUICommandInfo> AcceptActiveTool;
 	TSharedPtr<FUICommandInfo> CancelActiveTool;
diff --git a/Engine/Plugins/Experimental/UVEditor/Source/UVEditor/UVEditor.Build.cs b/Engine/Plugins/Experimental/UVEditor/Source/UVEditor/UVEditor.Build.cs
--- a/Engine/Plugins/Experimental/UVEditor/Source/UVEditor/UVEditor.Build.cs
+++ b/Engine/Plugins/Experimental/UVEditor/Source/UVEditor/UVEditor.Build.cs
@@ -58,6 +58,7 @@ public class UVEditor : ModuleRules
 				"ModelingComponents",
 				"ToolMenus",
 				"UVEditorTools",
+				"UVEditorToolsEditorOnly",
 				"ToolWidgets",				
 				// ... add private dependencies that you statically link with here ...	
 			}
diff --git a/Engine/Plugins/Experimental/UVEditor/Source/UVEditorTools/Private/ToolTargets/UVEditorToolMeshInput.cpp b/Engine/Plugins/Experimental/UVEditor/Source/UVEditorTools/Private/ToolTargets/UVEditorToolMeshInput.cpp
--- a/Engine/Plugins/Experimental/UVEditor/Source/UVEditorTools/Private/ToolTargets/UVEditorToolMeshInput.cpp
+++ b/Engine/Plugins/Experimental/UVEditor/Source/UVEditorTools/Private/ToolTargets/UVEditorToolMeshInput.cpp
@@ -56,7 +56,7 @@ namespace UVEditorToolMeshInputLocals
 
 		if (ChangedConnectivityTids)
 		{
-			for (int32 Tid : *ChangedConnectivityTids                                         )
+			for (int32 Tid : *ChangedConnectivityTids)
 			{
 				MeshOut.SetTriangle(Tid, MeshIn.GetTriangle(Tid));
 			}
@@ -453,5 +453,9 @@ void UUVEditorToolMeshInput::UpdateFromCanonicalUnwrapUsingMeshChange(const FDyn
 	VertexToTriangleOneRing(UnwrapCanonical.Get(), ChangedVids, RenderUpdateTidsSet);
 	TArray<int32> RenderUpdateTids = RenderUpdateTidsSet.Array();
 
-	UpdateAllFromUnwrapCanonical(&ChangedVids, &ChangedTids, &RenderUpdateTids);
+	// Todo: Determine a better solution for supporting explicit changed Tids and Vids
+	// 	     from MeshChanges. This current approach only works well if the topology is not
+	// 	     changing between updates.
+	//UpdateAllFromUnwrapCanonical(&ChangedVids, & ChangedTids, & RenderUpdateTids);
+	UpdateAllFromUnwrapCanonical(nullptr, nullptr, nullptr);
 }
diff --git a/Engine/Plugins/Experimental/UVEditor/Source/UVEditorTools/Private/UVEditorLayoutTool.cpp b/Engine/Plugins/Experimental/UVEditor/Source/UVEditorTools/Private/UVEditorLayoutTool.cpp
--- a/Engine/Plugins/Experimental/UVEditor/Source/UVEditorTools/Private/UVEditorLayoutTool.cpp
+++ b/Engine/Plugins/Experimental/UVEditor/Source/UVEditorTools/Private/UVEditorLayoutTool.cpp
@@ -119,10 +119,6 @@ void UUVEditorLayoutTool::Shutdown(EToolShutdownType ShutdownType)
 
 void UUVEditorLayoutTool::OnTick(float DeltaTime)
 {
-	for (TObjectPtr<UUVEditorToolMeshInput> Target : Targets)
-	{
-		Target->AppliedPreview->Tick(DeltaTime);
-	}
 }
 
 
diff --git a/Engine/Plugins/Experimental/UVEditor/Source/UVEditorToolsEditorOnly/Private/UVEditorParameterizeMeshTool.cpp b/Engine/Plugins/Experimental/UVEditor/Source/UVEditorToolsEditorOnly/Private/UVEditorParameterizeMeshTool.cpp
new file mode 100644
--- /dev/null
+++ b/Engine/Plugins/Experimental/UVEditor/Source/UVEditorToolsEditorOnly/Private/UVEditorParameterizeMeshTool.cpp
@@ -0,0 +1,210 @@
+// Copyright Epic Games, Inc. All Rights Reserved.
+
+#include "UVEditorParameterizeMeshTool.h"
+#include "InteractiveToolManager.h"
+#include "ToolBuilderUtil.h"
+#include "ToolSetupUtil.h"
+#include "ModelingToolTargetUtil.h"
+#include "DynamicMesh/DynamicMesh3.h"
+#include "FaceGroupUtil.h"
+#include "ParameterizationOps/ParameterizeMeshOp.h"
+#include "Properties/ParameterizeMeshProperties.h"
+#include "ToolTargets/UVEditorToolMeshInput.h"
+#include "UVToolContextObjects.h"
+#include "ContextObjectStore.h"
+
+#include "ExplicitUseGeometryMathTypes.h"		// using UE::Geometry::(math types)
+using namespace UE::Geometry;
+
+#define LOCTEXT_NAMESPACE "UParameterizeMeshTool"
+
+
+// Tool builder
+// TODO: Could consider sharing some of the tool builder boilerplate for UV editor tools in a common base class.
+
+const FToolTargetTypeRequirements& UUVEditorParameterizeMeshToolBuilder::GetTargetRequirements() const
+{
+	static FToolTargetTypeRequirements TypeRequirements(UUVEditorToolMeshInput::StaticClass());
+	return TypeRequirements;
+}
+
+bool UUVEditorParameterizeMeshToolBuilder::CanBuildTool(const FToolBuilderState& SceneState) const
+{
+	return Targets && Targets->Num() >= 1;
+}
+
+UInteractiveTool* UUVEditorParameterizeMeshToolBuilder::BuildTool(const FToolBuilderState& SceneState) const
+{
+	UUVEditorParameterizeMeshTool* NewTool = NewObject<UUVEditorParameterizeMeshTool>(SceneState.ToolManager);
+	NewTool->SetTargets(*Targets);
+
+	return NewTool;
+}
+
+
+/*
+ * Tool
+ */
+
+
+void UUVEditorParameterizeMeshTool::Setup()
+{
+	check(Targets.Num() >= 1);
+
+	UInteractiveTool::Setup();
+
+	// initialize our properties
+	Settings = NewObject<UParameterizeMeshToolProperties>(this);
+	Settings->RestoreProperties(this);
+	AddToolPropertySource(Settings);
+	Settings->WatchProperty(Settings->Method, [&](EParameterizeMeshUVMethod) { OnMethodTypeChanged(); });
+
+	UVAtlasProperties = NewObject<UParameterizeMeshToolUVAtlasProperties>(this);
+	UVAtlasProperties->RestoreProperties(this);
+	AddToolPropertySource(UVAtlasProperties);
+	SetToolPropertySourceEnabled(UVAtlasProperties, true);
+
+	XAtlasProperties = NewObject<UParameterizeMeshToolXAtlasProperties>(this);
+	XAtlasProperties->RestoreProperties(this);
+	AddToolPropertySource(XAtlasProperties);
+	SetToolPropertySourceEnabled(XAtlasProperties, true);
+	
+	PatchBuilderProperties = NewObject<UParameterizeMeshToolPatchBuilderProperties>(this);
+	PatchBuilderProperties->RestoreProperties(this);
+	AddToolPropertySource(PatchBuilderProperties);
+	SetToolPropertySourceEnabled(PatchBuilderProperties, true);
+
+	MaterialSettings = NewObject<UExistingMeshMaterialProperties>(this);
+	MaterialSettings->MaterialMode = ESetMeshMaterialMode::Checkerboard;
+	MaterialSettings->RestoreProperties(this, TEXT("ModelingUVTools"));
+	MaterialSettings->UVChannel = 0; // We only ever have one UV layer for the applied preview.
+	AddToolPropertySource(MaterialSettings);
+	// force update
+	MaterialSettings->UpdateMaterials();
+	for (TObjectPtr<UUVEditorToolMeshInput> Target : Targets)
+	{
+		Target->AppliedPreview->OverrideMaterial = MaterialSettings->GetActiveOverrideMaterial();
+	}
+
+	for (TObjectPtr<UUVEditorToolMeshInput> Target : Targets)
+	{
+		UParameterizeMeshOperatorFactory* OpFactory = NewObject<UParameterizeMeshOperatorFactory>();
+		OpFactory->TargetTransform = Target->AppliedPreview->PreviewMesh->GetTransform();
+		OpFactory->Settings = Settings;
+		OpFactory->UVAtlasProperties = UVAtlasProperties;
+		OpFactory->XAtlasProperties = XAtlasProperties;
+		OpFactory->PatchBuilderProperties = PatchBuilderProperties;
+		OpFactory->OriginalMesh = Target->AppliedCanonical;
+		OpFactory->GetSelectedUVChannel = [Target]() { return Target->UVLayerIndex; };
+
+		Target->AppliedPreview->ChangeOpFactory(OpFactory);
+		Target->AppliedPreview->OnMeshUpdated.AddWeakLambda(this, [Target](UMeshOpPreviewWithBackgroundCompute* Preview) {
+			Target->UpdateUnwrapPreviewFromAppliedPreview();
+			});
+
+		Target->AppliedPreview->InvalidateResult();
+	}
+
+	SetToolDisplayName(LOCTEXT("ToolNameGlobal", "AutoUV"));
+	GetToolManager()->DisplayMessage(
+		LOCTEXT("OnStartTool_Global", "Automatically partition the selected Mesh into UV islands, flatten, and pack into a single UV chart"),
+		EToolMessageLevel::UserNotification);
+}
+
+void UUVEditorParameterizeMeshTool::OnPropertyModified(UObject* PropertySet, FProperty* Property)
+{
+	if (PropertySet != MaterialSettings)
+	{
+		for (TObjectPtr<UUVEditorToolMeshInput> Target : Targets)
+		{
+			Target->AppliedPreview->InvalidateResult();
+		}
+	}
+
+	MaterialSettings->UpdateMaterials();
+	for (TObjectPtr<UUVEditorToolMeshInput> Target : Targets)
+	{
+		Target->AppliedPreview->OverrideMaterial = MaterialSettings->GetActiveOverrideMaterial();
+	}
+}
+
+
+void UUVEditorParameterizeMeshTool::OnMethodTypeChanged()
+{
+	SetToolPropertySourceEnabled(UVAtlasProperties, Settings->Method == EParameterizeMeshUVMethod::UVAtlas);
+	SetToolPropertySourceEnabled(XAtlasProperties, Settings->Method == EParameterizeMeshUVMethod::XAtlas);
+	SetToolPropertySourceEnabled(PatchBuilderProperties, Settings->Method == EParameterizeMeshUVMethod::PatchBuilder);
+
+	for (TObjectPtr<UUVEditorToolMeshInput> Target : Targets)
+	{
+		Target->AppliedPreview->InvalidateResult();
+	}
+}
+
+
+void UUVEditorParameterizeMeshTool::Shutdown(EToolShutdownType ShutdownType)
+{
+	Settings->SaveProperties(this);
+	MaterialSettings->SaveProperties(this, TEXT("ModelingUVTools"));
+	UVAtlasProperties->SaveProperties(this);
+	XAtlasProperties->SaveProperties(this);
+	PatchBuilderProperties->SaveProperties(this);
+
+	for (TObjectPtr<UUVEditorToolMeshInput> Target : Targets)
+	{
+		Target->AppliedPreview->OnMeshUpdated.RemoveAll(this);
+		Target->AppliedPreview->ClearOpFactory();
+		Target->AppliedPreview->OverrideMaterial = nullptr;
+	}
+
+	if (ShutdownType == EToolShutdownType::Accept)
+	{
+		UUVToolEmitChangeAPI* ChangeAPI = GetToolManager()->GetContextObjectStore()->FindContext<UUVToolEmitChangeAPI>();
+
+		for (TObjectPtr<UUVEditorToolMeshInput> Target : Targets)
+		{
+			// Set things up for undo. 
+			FDynamicMeshChangeTracker ChangeTracker(Target->UnwrapCanonical.Get());
+			ChangeTracker.BeginChange();
+
+			for (int32 Tid : Target->UnwrapCanonical->TriangleIndicesItr())
+			{
+				ChangeTracker.SaveTriangle(Tid, true);
+			}
+
+			Target->UpdateCanonicalFromPreviews();
+
+			ChangeAPI->EmitToolIndependentUnwrapCanonicalChange(Target, ChangeTracker.EndChange(), LOCTEXT("ApplyParameterizeMeshTool", "Auto UV Tool"));
+		}
+	}
+	else
+	{
+		// Reset the inputs
+		for (TObjectPtr<UUVEditorToolMeshInput> Target : Targets)
+		{
+			Target->UpdatePreviewsFromCanonical();
+		}
+	}
+
+	Settings = nullptr;
+	Targets.Empty();
+}
+
+void UUVEditorParameterizeMeshTool::OnTick(float DeltaTime)
+{
+}
+
+bool UUVEditorParameterizeMeshTool::CanAccept() const
+{
+	for (TObjectPtr<UUVEditorToolMeshInput> Target : Targets)
+	{
+		if (!Target->AppliedPreview->HaveValidResult())
+		{
+			return false;
+		}
+	}
+	return true;
+}
+
+
+#undef LOCTEXT_NAMESPACE
diff --git a/Engine/Plugins/Experimental/UVEditor/Source/UVEditorToolsEditorOnly/Private/UVEditorToolsEditorOnlyModule.cpp b/Engine/Plugins/Experimental/UVEditor/Source/UVEditorToolsEditorOnly/Private/UVEditorToolsEditorOnlyModule.cpp
new file mode 100644
--- /dev/null
+++ b/Engine/Plugins/Experimental/UVEditor/Source/UVEditorToolsEditorOnly/Private/UVEditorToolsEditorOnlyModule.cpp
@@ -0,0 +1,20 @@
+// Copyright Epic Games, Inc. All Rights Reserved.
+
+#include "UVEditorToolsEditorOnlyModule.h"
+
+#define LOCTEXT_NAMESPACE "FUVEditorToolsEditorOnlyModule"
+
+void FUVEditorToolsEditorOnlyModule::StartupModule()
+{
+	// This code will execute after your module is loaded into memory; the exact timing is specified in the .uplugin file per-module
+}
+
+void FUVEditorToolsEditorOnlyModule::ShutdownModule()
+{
+	// This function may be called during shutdown to clean up your module.  For modules that support dynamic reloading,
+	// we call this function before unloading the module.
+}
+
+#undef LOCTEXT_NAMESPACE
+
+IMPLEMENT_MODULE(FUVEditorToolsEditorOnlyModule, UVEditorToolsEditorOnly)
diff --git a/Engine/Plugins/Experimental/UVEditor/Source/UVEditorToolsEditorOnly/Public/UVEditorParameterizeMeshTool.h b/Engine/Plugins/Experimental/UVEditor/Source/UVEditorToolsEditorOnly/Public/UVEditorParameterizeMeshTool.h
new file mode 100644
--- /dev/null
+++ b/Engine/Plugins/Experimental/UVEditor/Source/UVEditorToolsEditorOnly/Public/UVEditorParameterizeMeshTool.h
@@ -0,0 +1,95 @@
+// Copyright Epic Games, Inc. All Rights Reserved.
+
+#pragma once
+
+#include "CoreMinimal.h"
+
+#include "InteractiveTool.h"
+#include "InteractiveToolBuilder.h"
+
+#include "DynamicMesh/DynamicMesh3.h"
+#include "MeshOpPreviewHelpers.h"
+#include "Properties/MeshMaterialProperties.h"
+#include "Properties/MeshUVChannelProperties.h"
+
+#include "UVEditorParameterizeMeshTool.generated.h"
+
+// predeclarations
+class UDynamicMeshComponent;
+class UMaterialInterface;
+class UMaterialInstanceDynamic;
+class UUVEditorToolMeshInput;
+class UParameterizeMeshToolProperties;
+class UParameterizeMeshToolUVAtlasProperties;
+class UParameterizeMeshToolXAtlasProperties;
+class UParameterizeMeshToolPatchBuilderProperties;
+
+UCLASS()
+class UVEDITORTOOLSEDITORONLY_API UUVEditorParameterizeMeshToolBuilder : public UInteractiveToolBuilder
+{
+	GENERATED_BODY()
+public:
+	virtual bool CanBuildTool(const FToolBuilderState& SceneState) const override;
+	virtual UInteractiveTool* BuildTool(const FToolBuilderState& SceneState) const override;
+
+	// This is a pointer so that it can be updated under the builder without
+	// having to set it in the mode after initializing targets.
+	const TArray<TObjectPtr<UUVEditorToolMeshInput>>* Targets = nullptr;
+
+protected:
+	virtual const FToolTargetTypeRequirements& GetTargetRequirements() const override;
+};
+
+
+/**
+ * UParameterizeMeshTool automatically decomposes the input mesh into charts, solves for UVs,
+ * and then packs the resulting charts
+ */
+UCLASS()
+class UVEDITORTOOLSEDITORONLY_API UUVEditorParameterizeMeshTool : public UInteractiveTool
+{
+	GENERATED_BODY()
+
+public:
+	/**
+     * The tool will operate on the meshes given here.
+     */
+	virtual void SetTargets(const TArray<TObjectPtr<UUVEditorToolMeshInput>>& TargetsIn)
+	{
+		Targets = TargetsIn;
+	}
+
+
+	virtual void Setup() override;
+	virtual void Shutdown(EToolShutdownType ShutdownType) override;
+
+	virtual void OnTick(float DeltaTime) override;
+
+	virtual bool HasCancel() const override { return true; }
+	virtual bool HasAccept() const override { return true; }
+	virtual bool CanAccept() const override;
+
+	virtual void OnPropertyModified(UObject* PropertySet, FProperty* Property) override;
+
+protected:
+	UPROPERTY()
+	TArray<TObjectPtr<UUVEditorToolMeshInput>> Targets;
+
+	UPROPERTY()
+	TObjectPtr<UParameterizeMeshToolProperties> Settings = nullptr;
+
+	UPROPERTY()
+	TObjectPtr<UParameterizeMeshToolUVAtlasProperties> UVAtlasProperties = nullptr;
+
+	UPROPERTY()
+	TObjectPtr<UParameterizeMeshToolXAtlasProperties> XAtlasProperties = nullptr;
+
+	UPROPERTY()
+	TObjectPtr<UParameterizeMeshToolPatchBuilderProperties> PatchBuilderProperties = nullptr;
+
+	UPROPERTY()
+	TObjectPtr<UExistingMeshMaterialProperties> MaterialSettings = nullptr;
+
+	void OnMethodTypeChanged();
+};
+
diff --git a/Engine/Plugins/Experimental/UVEditor/Source/UVEditorToolsEditorOnly/Public/UVEditorToolsEditorOnlyModule.h b/Engine/Plugins/Experimental/UVEditor/Source/UVEditorToolsEditorOnly/Public/UVEditorToolsEditorOnlyModule.h
new file mode 100644
--- /dev/null
+++ b/Engine/Plugins/Experimental/UVEditor/Source/UVEditorToolsEditorOnly/Public/UVEditorToolsEditorOnlyModule.h
@@ -0,0 +1,15 @@
+// Copyright Epic Games, Inc. All Rights Reserved.
+
+#pragma once
+
+#include "CoreMinimal.h"
+#include "Modules/ModuleManager.h"
+
+class FUVEditorToolsEditorOnlyModule : public IModuleInterface
+{
+public:
+
+	/** IModuleInterface implementation */
+	virtual void StartupModule() override;
+	virtual void ShutdownModule() override;
+};
\ No newline at end of file
diff --git a/Engine/Plugins/Experimental/UVEditor/Source/UVEditorToolsEditorOnly/UVEditorToolsEditorOnly.Build.cs b/Engine/Plugins/Experimental/UVEditor/Source/UVEditorToolsEditorOnly/UVEditorToolsEditorOnly.Build.cs
new file mode 100644
--- /dev/null
+++ b/Engine/Plugins/Experimental/UVEditor/Source/UVEditorToolsEditorOnly/UVEditorToolsEditorOnly.Build.cs
@@ -0,0 +1,67 @@
+// Copyright Epic Games, Inc. All Rights Reserved.
+
+using UnrealBuildTool;
+
+public class UVEditorToolsEditorOnly : ModuleRules
+{
+	public UVEditorToolsEditorOnly(ReadOnlyTargetRules Target) : base(Target)
+	{
+		PCHUsage = ModuleRules.PCHUsageMode.UseExplicitOrSharedPCHs;
+		
+		PublicIncludePaths.AddRange(
+			new string[] {
+				// ... add public include paths required here ...
+			}
+			);
+				
+		
+		PrivateIncludePaths.AddRange(
+			new string[] {
+				// ... add other private include paths required here ...
+			}
+			);
+			
+		
+		PublicDependencyModuleNames.AddRange(
+			new string[]
+			{
+				"Core",
+				
+				"InteractiveToolsFramework",
+				
+				// ... add other public dependencies that you statically link with here ...
+			}
+			);
+			
+		
+		PrivateDependencyModuleNames.AddRange(
+			new string[]
+			{
+				// NOTE: UVEditorTools is a separate module so that it doesn't rely on the editor.
+				// So, do not add editor dependencies here.
+				
+				"CoreUObject",
+				"Engine",
+				"Slate",
+				"SlateCore",
+					
+				"MeshModelingToolsExp",
+				"DynamicMesh",
+				"GeometryCore",
+				"ModelingComponents",
+				"ModelingOperators",
+				"ModelingOperatorsEditorOnly",
+				"UVEditorTools",
+				// ... add private dependencies that you statically link with here ...	
+			}
+			);
+		
+		
+		DynamicallyLoadedModuleNames.AddRange(
+			new string[]
+			{
+				// ... add any modules that your module loads dynamically here ...
+			}
+			);
+	}
+}
diff --git a/Engine/Plugins/Experimental/UVEditor/UVEditor.uplugin b/Engine/Plugins/Experimental/UVEditor/UVEditor.uplugin
--- a/Engine/Plugins/Experimental/UVEditor/UVEditor.uplugin
+++ b/Engine/Plugins/Experimental/UVEditor/UVEditor.uplugin
@@ -24,6 +24,12 @@
 			"Name": "UVEditorTools",
 			"Type": "Editor",
 			"LoadingPhase": "Default"
+		},
+		{
+			"Name": "UVEditorToolsEditorOnly",
+			"Type": "Editor",
+			"LoadingPhase": "Default",
+            "WhitelistPlatforms": [ "Win64" ]
 		}
 	],
 	"Plugins": [
