UVEditor: Add UV viewport toolbar with gizmo controls that can be used in the Select tool. With this in place, removed the Transform tool, and removed inter-invocation selection storage in the Select tool (which mainly existed to support a separate Transform tool). Removed some hacks that were in place to make a separate Transform tool possible.

#rnx

diff --git a/Engine/Plugins/Experimental/UVEditor/Source/UVEditor/Private/SUVEditor2DViewport.cpp b/Engine/Plugins/Experimental/UVEditor/Source/UVEditor/Private/SUVEditor2DViewport.cpp
--- a/Engine/Plugins/Experimental/UVEditor/Source/UVEditor/Private/SUVEditor2DViewport.cpp
+++ b/Engine/Plugins/Experimental/UVEditor/Source/UVEditor/Private/SUVEditor2DViewport.cpp
@@ -2,6 +2,9 @@
 
 #include "SUVEditor2DViewport.h"
 
+#include "SUVEditor2DViewportToolBar.h"
+#include "UVEditor2DViewportClient.h"
+
 #define LOCTEXT_NAMESPACE "SUVEditor2DViewport"
 
 void SUVEditor2DViewport::AddOverlayWidget(TSharedRef<SWidget> OverlaidWidget)
@@ -17,4 +20,16 @@ void SUVEditor2DViewport::RemoveOverlayWidget(TSharedRef<SWidget> OverlaidWidget
 	ViewportOverlay->RemoveSlot(OverlaidWidget);
 }
 
+TSharedPtr<SWidget> SUVEditor2DViewport::MakeViewportToolbar()
+{
+	return SNew(SUVEditor2DViewportToolBar)
+		.CommandList(CommandList);
+}
+
+bool SUVEditor2DViewport::IsWidgetModeActive(UE::Widget::EWidgetMode Mode) const
+{
+	return static_cast<FUVEditor2DViewportClient*>(Client.Get())->AreWidgetButtonsEnabled() 
+		&& Client->GetWidgetMode() == Mode;
+}
+
 #undef LOCTEXT_NAMESPACE
diff --git a/Engine/Plugins/Experimental/UVEditor/Source/UVEditor/Private/SUVEditor2DViewportToolBar.cpp b/Engine/Plugins/Experimental/UVEditor/Source/UVEditor/Private/SUVEditor2DViewportToolBar.cpp
new file mode 100644
--- /dev/null
+++ b/Engine/Plugins/Experimental/UVEditor/Source/UVEditor/Private/SUVEditor2DViewportToolBar.cpp
@@ -0,0 +1,76 @@
+// Copyright Epic Games, Inc. All Rights Reserved.
+
+#pragma once
+
+#include "SUVEditor2DViewportToolBar.h"
+
+#include "EditorViewportCommands.h"
+
+
+void SUVEditor2DViewportToolBar::Construct(const FArguments& InArgs)
+{
+	CommandList = InArgs._CommandList;
+
+	const FMargin ToolbarSlotPadding(4.0f, 1.0f);
+	TSharedPtr<SHorizontalBox> MainBoxPtr;
+
+	ChildSlot
+	[
+		SNew( SBorder )
+		.BorderImage(FAppStyle::Get().GetBrush("EditorViewportToolBar.Background"))
+		.Cursor(EMouseCursor::Default)
+		[
+			SAssignNew( MainBoxPtr, SHorizontalBox )
+		]
+	];
+
+	// Transform toolbar
+	MainBoxPtr->AddSlot()
+		.Padding(ToolbarSlotPadding)
+		.HAlign(HAlign_Right)
+		[
+			MakeToolBar(InArgs._Extenders)
+		];
+
+	SViewportToolBar::Construct(SViewportToolBar::FArguments());
+}
+
+TSharedRef<SWidget> SUVEditor2DViewportToolBar::MakeToolBar(const TSharedPtr<FExtender> InExtenders)
+{
+	// The following is modeled after portions of STransformViewportToolBar, which gets 
+	// used in SCommonEditorViewportToolbarBase.
+
+	// The buttons are hooked up to actual functions via command bindings in SUVEditor2DViewport::BindCommands(),
+	// and the toolbar gets built in SUVEditor2DViewport::MakeViewportToolbar().
+
+	FSlimHorizontalToolBarBuilder ToolbarBuilder(CommandList, FMultiBoxCustomization::None, InExtenders);
+
+	// Use a custom style
+	FName ToolBarStyle = "EditorViewportToolBar";
+	ToolbarBuilder.SetStyle(&FEditorStyle::Get(), ToolBarStyle);
+	ToolbarBuilder.SetLabelVisibility(EVisibility::Collapsed);
+
+	// Transform controls should not be focusable as it fights with the press space to change transform 
+	// mode feature, which we may someday have.
+	ToolbarBuilder.SetIsFocusable(false);
+
+	// Widget controls
+	ToolbarBuilder.BeginSection("Transform");
+	{
+		ToolbarBuilder.BeginBlockGroup();
+
+		// Select Mode
+		static FName SelectModeName = FName(TEXT("SelectMode"));
+		ToolbarBuilder.AddToolBarButton(FEditorViewportCommands::Get().SelectMode, NAME_None, TAttribute<FText>(), TAttribute<FText>(), TAttribute<FSlateIcon>(), SelectModeName);
+
+		// Translate Mode
+		static FName TranslateModeName = FName(TEXT("TranslateMode"));
+		ToolbarBuilder.AddToolBarButton(FEditorViewportCommands::Get().TranslateMode, NAME_None, TAttribute<FText>(), TAttribute<FText>(), TAttribute<FSlateIcon>(), TranslateModeName);
+
+		ToolbarBuilder.EndBlockGroup();
+	}
+
+	ToolbarBuilder.EndSection();
+
+	return ToolbarBuilder.MakeWidget();
+}
\ No newline at end of file
diff --git a/Engine/Plugins/Experimental/UVEditor/Source/UVEditor/Private/UVEditor2DViewportClient.cpp b/Engine/Plugins/Experimental/UVEditor/Source/UVEditor/Private/UVEditor2DViewportClient.cpp
--- a/Engine/Plugins/Experimental/UVEditor/Source/UVEditor/Private/UVEditor2DViewportClient.cpp
+++ b/Engine/Plugins/Experimental/UVEditor/Source/UVEditor/Private/UVEditor2DViewportClient.cpp
@@ -2,19 +2,21 @@
 
 #include "UVEditor2DViewportClient.h"
 
-#include "UVEditorMode.h"
 #include "BaseBehaviors/ClickDragBehavior.h"
 #include "BaseBehaviors/MouseWheelBehavior.h"
-//#include "Drawing/DensityAdjustingGrid.h"
+#include "ContextObjectStore.h"
 #include "EditorModeManager.h"
 #include "EdModeInteractiveToolsContext.h"
 #include "Drawing/MeshDebugDrawing.h"
 #include "FrameTypes.h"
+#include "UVEditorMode.h"
 
 FUVEditor2DViewportClient::FUVEditor2DViewportClient(FEditorModeTools* InModeTools,
 	FPreviewScene* InPreviewScene, const TWeakPtr<SEditorViewport>& InEditorViewportWidget)
 	: FEditorViewportClient(InModeTools, InPreviewScene, InEditorViewportWidget)
 {
+	ShowWidget(false);
+
 	// Don't draw the little XYZ drawing in the corner.
 	bDrawAxes = false;
 
@@ -43,6 +45,14 @@ FUVEditor2DViewportClient::FUVEditor2DViewportClient(FEditorModeTools* InModeToo
 	BehaviorSet->Add(ZoomBehavior);
 
 	ModeTools->GetInteractiveToolsContext()->InputRouter->RegisterSource(this);
+
+	UContextObjectStore* ContextStore = ModeTools->GetInteractiveToolsContext()->ContextObjectStore;
+	ViewportButtonsAPI = ContextStore->FindContext<UUVToolViewportButtonsAPI>();
+	if (!ViewportButtonsAPI)
+	{
+		ViewportButtonsAPI = NewObject<UUVToolViewportButtonsAPI>();
+		ContextStore->AddContextObject(ViewportButtonsAPI);
+	}
 }
 
 const UInputBehaviorSet* FUVEditor2DViewportClient::GetInputBehaviors() const
@@ -56,7 +66,7 @@ void FUVEditor2DViewportClient::AddReferencedObjects(FReferenceCollector& Collec
 	FEditorViewportClient::AddReferencedObjects(Collector);
 
 	Collector.AddReferencedObject(BehaviorSet);
-	//Collector.AddReferencedObject(Grid);
+	Collector.AddReferencedObject(ViewportButtonsAPI);
 }
 
 bool FUVEditor2DViewportClient::InputKey(FViewport* InViewport, int32 ControllerId, FKey Key, EInputEvent Event, float/*AmountDepressed*/, bool/*Gamepad*/)
@@ -135,3 +145,59 @@ bool FUVEditor2DViewportClient::ShouldOrbitCamera() const
 {
 	return false; // The UV Editor's 2D viewport should never orbit.
 }
+
+void FUVEditor2DViewportClient::SetWidgetMode(UE::Widget::EWidgetMode NewMode)
+{
+	if (ViewportButtonsAPI)
+	{
+		switch (NewMode)
+		{
+		case UE::Widget::EWidgetMode::WM_None:
+			ViewportButtonsAPI->SetGizmoMode(UUVToolViewportButtonsAPI::EGizmoMode::Select);
+			break;
+		case UE::Widget::EWidgetMode::WM_Translate:
+			ViewportButtonsAPI->SetGizmoMode(UUVToolViewportButtonsAPI::EGizmoMode::Transform);
+			break;
+		default:
+			// Do nothing
+			break;
+		}
+	}
+}
+
+bool FUVEditor2DViewportClient::AreWidgetButtonsEnabled() const
+{
+	return ViewportButtonsAPI && ViewportButtonsAPI->AreGizmoButtonsEnabled();
+}
+
+bool FUVEditor2DViewportClient::CanSetWidgetMode(UE::Widget::EWidgetMode NewMode) const
+{
+	if (!AreWidgetButtonsEnabled())
+	{
+		return false;
+	}
+
+	return NewMode == UE::Widget::EWidgetMode::WM_None 
+		|| NewMode == UE::Widget::EWidgetMode::WM_Translate;
+}
+
+UE::Widget::EWidgetMode FUVEditor2DViewportClient::GetWidgetMode() const
+{
+	if (!AreWidgetButtonsEnabled())
+	{
+		return UE::Widget::EWidgetMode::WM_None;
+	}
+
+	switch (ViewportButtonsAPI->GetGizmoMode())
+	{
+	case UUVToolViewportButtonsAPI::EGizmoMode::Select:
+		return UE::Widget::EWidgetMode::WM_None;
+		break;
+	case UUVToolViewportButtonsAPI::EGizmoMode::Transform:
+		return UE::Widget::EWidgetMode::WM_Translate;
+		break;
+	default:
+		return UE::Widget::EWidgetMode::WM_None;
+		break;
+	}
+}
diff --git a/Engine/Plugins/Experimental/UVEditor/Source/UVEditor/Private/UVEditorCommands.cpp b/Engine/Plugins/Experimental/UVEditor/Source/UVEditor/Private/UVEditorCommands.cpp
--- a/Engine/Plugins/Experimental/UVEditor/Source/UVEditor/Private/UVEditorCommands.cpp
+++ b/Engine/Plugins/Experimental/UVEditor/Source/UVEditor/Private/UVEditorCommands.cpp
@@ -24,7 +24,6 @@ void FUVEditorCommands::RegisterCommands()
 
 	// These get linked to various tool buttons.
 	UI_COMMAND(BeginSelectTool, "Select", "Switch to select tool", EUserInterfaceActionType::ToggleButton, FInputChord());
-	UI_COMMAND(BeginTransformTool, "Transform", "Switch to transform tool", EUserInterfaceActionType::ToggleButton, FInputChord());
 	UI_COMMAND(BeginLayoutTool, "Layout", "Switch to layout tool", EUserInterfaceActionType::ToggleButton, FInputChord());
 	UI_COMMAND(BeginParameterizeMeshTool, "Auto UV", "Switch to Auto UV tool", EUserInterfaceActionType::ToggleButton, FInputChord());
 
diff --git a/Engine/Plugins/Experimental/UVEditor/Source/UVEditor/Private/UVEditorMode.cpp b/Engine/Plugins/Experimental/UVEditor/Source/UVEditor/Private/UVEditorMode.cpp
--- a/Engine/Plugins/Experimental/UVEditor/Source/UVEditor/Private/UVEditorMode.cpp
+++ b/Engine/Plugins/Experimental/UVEditor/Source/UVEditor/Private/UVEditorMode.cpp
@@ -84,16 +84,10 @@ void UUVEditorMode::RegisterTools()
 {
 	const FUVEditorCommands& CommandInfos = FUVEditorCommands::Get();
 
-	// TODO: Other tool registrations go here
 	UUVSelectToolBuilder* UVSelectToolBuilder = NewObject<UUVSelectToolBuilder>();
 	UVSelectToolBuilder->Targets = &ToolInputObjects;
 	RegisterTool(CommandInfos.BeginSelectTool, TEXT("UVSelectTool"), UVSelectToolBuilder);
 
-	UUVSelectToolBuilder* UVTransformToolBuilder = NewObject<UUVSelectToolBuilder>();
-	UVTransformToolBuilder->Targets = &ToolInputObjects;
-	UVTransformToolBuilder->bGizmoEnabled = true;
-	RegisterTool(CommandInfos.BeginTransformTool, TEXT("UVTransformTool"), UVTransformToolBuilder);
-
 	UUVEditorLayoutToolBuilder* UVEditorLayoutToolBuilder = NewObject<UUVEditorLayoutToolBuilder>();
 	UVEditorLayoutToolBuilder->Targets = &ToolInputObjects;
 	RegisterTool(CommandInfos.BeginLayoutTool, TEXT("UVLayoutTool"), UVEditorLayoutToolBuilder);
diff --git a/Engine/Plugins/Experimental/UVEditor/Source/UVEditor/Private/UVEditorModeToolkit.cpp b/Engine/Plugins/Experimental/UVEditor/Source/UVEditor/Private/UVEditorModeToolkit.cpp
--- a/Engine/Plugins/Experimental/UVEditor/Source/UVEditor/Private/UVEditorModeToolkit.cpp
+++ b/Engine/Plugins/Experimental/UVEditor/Source/UVEditor/Private/UVEditorModeToolkit.cpp
@@ -183,7 +183,6 @@ void FUVEditorModeToolkit::Init(const TSharedPtr<IToolkitHost>& InitToolkitHost,
 
 	// TODO: Add more tools here
 	ToolbarBuilder.AddToolBarButton(CommandInfos.BeginSelectTool);
-	ToolbarBuilder.AddToolBarButton(CommandInfos.BeginTransformTool);
 	ToolbarBuilder.AddToolBarButton(CommandInfos.BeginLayoutTool);
 	ToolbarBuilder.AddToolBarButton(CommandInfos.BeginParameterizeMeshTool);
 
diff --git a/Engine/Plugins/Experimental/UVEditor/Source/UVEditor/Public/SUVEditor2DViewport.h b/Engine/Plugins/Experimental/UVEditor/Source/UVEditor/Public/SUVEditor2DViewport.h
--- a/Engine/Plugins/Experimental/UVEditor/Source/UVEditor/Public/SUVEditor2DViewport.h
+++ b/Engine/Plugins/Experimental/UVEditor/Source/UVEditor/Public/SUVEditor2DViewport.h
@@ -16,4 +16,8 @@ public:
 	// is not helpful here because that gets called just once.
 	virtual void AddOverlayWidget(TSharedRef<SWidget> OverlaidWidget);
 	virtual void RemoveOverlayWidget(TSharedRef<SWidget> OverlaidWidget);
+
+	// SEditorViewport
+	virtual TSharedPtr<SWidget> MakeViewportToolbar() override;
+	bool IsWidgetModeActive(UE::Widget::EWidgetMode Mode) const override;
 };
diff --git a/Engine/Plugins/Experimental/UVEditor/Source/UVEditor/Public/SUVEditor2DViewportToolBar.h b/Engine/Plugins/Experimental/UVEditor/Source/UVEditor/Public/SUVEditor2DViewportToolBar.h
new file mode 100644
--- /dev/null
+++ b/Engine/Plugins/Experimental/UVEditor/Source/UVEditor/Public/SUVEditor2DViewportToolBar.h
@@ -0,0 +1,30 @@
+// Copyright Epic Games, Inc. All Rights Reserved.
+
+#pragma once
+
+#include "CoreMinimal.h"
+
+#include "SViewportToolBar.h"
+
+class FUICommandList;
+class SUVEditor2DViewport;
+
+/**
+ * Toolbar that shows up at the top of the 2d viewport (has gizmo controls)
+ */
+class SUVEditor2DViewportToolBar : public SViewportToolBar
+{
+public:
+	SLATE_BEGIN_ARGS(SUVEditor2DViewportToolBar) {}
+		SLATE_ARGUMENT(TSharedPtr<FUICommandList>, CommandList)
+		SLATE_ARGUMENT(TSharedPtr<FExtender>, Extenders)
+	SLATE_END_ARGS()
+
+	void Construct(const FArguments& InArgs);
+
+private:
+	TSharedRef<SWidget> MakeToolBar(const TSharedPtr<FExtender> InExtenders);
+
+	/** Command list */
+	TSharedPtr<FUICommandList> CommandList;
+};
diff --git a/Engine/Plugins/Experimental/UVEditor/Source/UVEditor/Public/UVEditor2DViewportClient.h b/Engine/Plugins/Experimental/UVEditor/Source/UVEditor/Public/UVEditor2DViewportClient.h
--- a/Engine/Plugins/Experimental/UVEditor/Source/UVEditor/Public/UVEditor2DViewportClient.h
+++ b/Engine/Plugins/Experimental/UVEditor/Source/UVEditor/Public/UVEditor2DViewportClient.h
@@ -7,8 +7,9 @@
 #include "EditorViewportClient.h"
 #include "InputBehaviorSet.h"
 #include "UVEditor2DViewportBehaviorTargets.h" // FUVEditor2DScrollBehaviorTarget, FUVEditor2DMouseWheelZoomBehaviorTarget
+#include "UVToolContextObjects.h" // UUVToolViewportButtonsAPI::EGizmoMode
 
-//class UDensityAdjustingGrid;
+class UUVToolViewportButtonsAPI;
 
 /**
  * Client used to display a 2D view of the UV's, implemented by using a perspective viewport with a locked
@@ -22,10 +23,15 @@ public:
 
 	virtual ~FUVEditor2DViewportClient() {}
 
+	bool AreWidgetButtonsEnabled() const;
+
 	// FEditorViewportClient
 	virtual bool InputKey(FViewport* InViewport, int32 ControllerId, FKey Key, EInputEvent Event, float/*AmountDepressed*/, bool/*Gamepad*/) override;
 	virtual void Draw(const FSceneView* View, FPrimitiveDrawInterface* PDI) override;
 	virtual bool ShouldOrbitCamera() const override;
+	bool CanSetWidgetMode(UE::Widget::EWidgetMode NewMode) const override;
+	void SetWidgetMode(UE::Widget::EWidgetMode NewMode) override;
+	UE::Widget::EWidgetMode GetWidgetMode() const override;
 
 	// IInputBehaviorSource
 	virtual const UInputBehaviorSet* GetInputBehaviors() const override;
@@ -36,7 +42,7 @@ public:
 protected:
 	// These get added in AddReferencedObjects for memory management
 	UInputBehaviorSet* BehaviorSet;
-	//UDensityAdjustingGrid* Grid;
+	UUVToolViewportButtonsAPI* ViewportButtonsAPI;
 
 	// Note that it's generally less hassle if the unique ptr types are complete here,
 	// not forward declared, else we get compile errors if their destruction shows up
diff --git a/Engine/Plugins/Experimental/UVEditor/Source/UVEditor/Public/UVEditorCommands.h b/Engine/Plugins/Experimental/UVEditor/Source/UVEditor/Public/UVEditorCommands.h
--- a/Engine/Plugins/Experimental/UVEditor/Source/UVEditor/Public/UVEditorCommands.h
+++ b/Engine/Plugins/Experimental/UVEditor/Source/UVEditor/Public/UVEditorCommands.h
@@ -19,7 +19,6 @@ public:
 	TSharedPtr<FUICommandInfo> ApplyChanges;
 
 	TSharedPtr<FUICommandInfo> BeginSelectTool;
-	TSharedPtr<FUICommandInfo> BeginTransformTool;
 	TSharedPtr<FUICommandInfo> BeginLayoutTool;
 	TSharedPtr<FUICommandInfo> BeginParameterizeMeshTool;
 
diff --git a/Engine/Plugins/Experimental/UVEditor/Source/UVEditorTools/Private/UVSelectTool.cpp b/Engine/Plugins/Experimental/UVEditor/Source/UVEditorTools/Private/UVSelectTool.cpp
--- a/Engine/Plugins/Experimental/UVEditor/Source/UVEditorTools/Private/UVSelectTool.cpp
+++ b/Engine/Plugins/Experimental/UVEditor/Source/UVEditorTools/Private/UVSelectTool.cpp
@@ -93,49 +93,6 @@ namespace UVSelectToolLocals
 		FTransform GizmoBefore;
 	};
 
-	/**
-	 * A change that allows another change to be injected into it even after it is placed in the undo
-	 * stack.
-	 * NOTE: Remove this if we merge Transform into Select tool and choose not to store/restore
-	 * selecitons in the Select tool.
-	 */
-	class FSpeculativeSelectionChange : public FToolCommandChange
-	{
-	public:
-		FSpeculativeSelectionChange(TSharedPtr<TUniquePtr<FToolCommandChange>> PointerToContentIn)
-			: PointerToContent(PointerToContentIn)
-		{
-		}
-
-		virtual bool HasExpired(UObject* Object) const override
-		{
-			return !PointerToContent->IsValid();
-		}
-
-		virtual void Apply(UObject* Object) override
-		{
-			if (*PointerToContent)
-			{
-				(*PointerToContent)->Apply(Object);
-			}
-		}
-
-		virtual void Revert(UObject* Object) override
-		{
-			if (*PointerToContent)
-			{
-				(*PointerToContent)->Revert(Object);
-			}
-		}
-
-		virtual FString ToString() const override
-		{
-			return TEXT("UVSelectToolLocals::FSpeculativeSelectionChange");
-		}
-
-		TSharedPtr<TUniquePtr<FToolCommandChange>> PointerToContent;
-	};
-
 	/**
 	 * A change similar to the one emitted by EmitChangeApi->EmitToolIndependentUnwrapCanonicalChange,
 	 * but which updates the Select tool's gizmo in a way that preserves the rotational component
@@ -219,24 +176,6 @@ namespace UVSelectToolLocals
 
 }
 
-void UUVSelectToolSpeculativeChangeAPI::EmitSpeculativeChange(UObject* TargetObject, 
-	UUVToolEmitChangeAPI* EmitChangeAPI, const FText& TransactionName)
-{
-	ContentOfLastSpeculativeChange = MakeShared<TUniquePtr<FToolCommandChange>>();
-
-	EmitChangeAPI->EmitToolIndependentChange(TargetObject, 
-		MakeUnique<UVSelectToolLocals::FSpeculativeSelectionChange>(ContentOfLastSpeculativeChange), 
-		TransactionName);
-}
-
-void UUVSelectToolSpeculativeChangeAPI::InsertIntoLastSpeculativeChange(TUniquePtr<FToolCommandChange> ChangeToInsert)
-{
-	if (ContentOfLastSpeculativeChange)
-	{
-		*ContentOfLastSpeculativeChange = MoveTemp(ChangeToInsert);
-	}
-}
-
 
 /*
  * ToolBuilder
@@ -257,7 +196,6 @@ UInteractiveTool* UUVSelectToolBuilder::BuildTool(const FToolBuilderState& Scene
 {
 	UUVSelectTool* NewTool = NewObject<UUVSelectTool>(SceneState.ToolManager);
 	NewTool->SetWorld(SceneState.World);
-	NewTool->SetGizmoEnabled(bGizmoEnabled);
 	NewTool->SetTargets(*Targets);
 
 	return NewTool;
@@ -293,6 +231,12 @@ void UUVSelectTool::Setup()
 
 	UContextObjectStore* ContextStore = GetToolManager()->GetContextObjectStore();
 	EmitChangeAPI = ContextStore->FindContext<UUVToolEmitChangeAPI>();
+	ViewportButtonsAPI = ContextStore->FindContext<UUVToolViewportButtonsAPI>();
+	ViewportButtonsAPI->SetGizmoButtonsEnabled(true);
+	ViewportButtonsAPI->OnGizmoModeChange.AddWeakLambda(this, 
+		[this](UUVToolViewportButtonsAPI::EGizmoMode NewGizmoMode) {
+			UpdateGizmo();
+		});
 
 	ToolActions = NewObject<USelectToolActionPropertySet>(this);
 	ToolActions->Initialize(this);
@@ -355,29 +299,6 @@ void UUVSelectTool::Setup()
 			Targets[i]->UnwrapPreview->PreviewMesh->GetTransform());
 	}
 
-	// See if we have a stored selection
-	UUVToolMeshSelection* SelectionStore = ContextStore->FindContext<UUVToolMeshSelection>();
-	if (SelectionStore)
-	{
-		SelectionMechanic->SetSelection(*SelectionStore->Selection);
-	}
-	if (!SelectionMechanic->GetCurrentSelection().IsEmpty() && !SelectionMechanic->GetCurrentSelection().MatchesTimestamp())
-	{
-		// If we have an expired selection, then there must have been some tool that changed the
-		// topology. We need to clear the selection in an undoable way, but that clear actually
-		// needs to happen before the last invocation shutdown to be in the proper sequence in the 
-		// undo stack.
-		FDynamicMeshSelection NewSelection;
-		UUVSelectToolSpeculativeChangeAPI* SpeculativeChangeAPI = ContextStore->FindContext<UUVSelectToolSpeculativeChangeAPI>();
-		if (SpeculativeChangeAPI && SpeculativeChangeAPI->HasSpeculativeChange())
-		{
-			SpeculativeChangeAPI->InsertIntoLastSpeculativeChange(MakeUnique<UVSelectToolLocals::FSelectionChange>(
-				SelectionMechanic->GetCurrentSelection(), NewSelection, false, FTransform()));
-		}
-
-		SelectionMechanic->SetSelection(NewSelection, false, false);
-	}
-
 	// Make sure that if we receive undo/redo events on the meshes, we update the tree structures
 	// and the selection mechanic drawn elements. Note that we mainly have to worry about this
 	// because the select tool is the default UV editor tool, and as such it can receive undo
@@ -408,6 +329,7 @@ void UUVSelectTool::Setup()
 	// Always align gizmo to x and y axes
 	TransformGizmo->bUseContextCoordinateSystem = false;
 	TransformGizmo->SetActiveTarget(TransformProxy, GetToolManager());
+	TransformGizmo->SetVisibility(ViewportButtonsAPI->GetGizmoMode() != UUVToolViewportButtonsAPI::EGizmoMode::Select);
 
 	LivePreviewGeometryActor = Targets[0]->AppliedPreview->GetWorld()->SpawnActor<APreviewGeometryActor>(
 		FVector::ZeroVector, FRotator(0, 0, 0), FActorSpawnParameters());
@@ -432,43 +354,11 @@ void UUVSelectTool::Setup()
 
 void UUVSelectTool::Shutdown(EToolShutdownType ShutdownType)
 {
-	UContextObjectStore* ContextStore = GetToolManager()->GetContextObjectStore();
+	// Clear selection so that it can be restored after undoing back into the select tool
 	if (!SelectionMechanic->GetCurrentSelection().IsEmpty())
 	{
-		// TODO: It's not clear whether storing our selection is the right choice in the long run,
-		// though we currently need it while the gizmo transformation is a separate tool... 
-		UUVToolMeshSelection* SelectionStore = ContextStore->FindContext<UUVToolMeshSelection>();
-		if (!SelectionStore)
-		{
-			SelectionStore = NewObject<UUVToolMeshSelection>();
-			ContextStore->AddContextObject(SelectionStore);
-		}
-		*SelectionStore->Selection = SelectionMechanic->GetCurrentSelection();
-
-		// We update the topology timestamp here because the timestamps may have divereged due
-		// to our undo/redo change implementation sometimes incrementing the topology timestamp
-		// even when the topology hasn't changed.
-		// TODO: Check whether we still need this once we use FMeshVertexChange instead of FDynamicMeshChange
-		SelectionStore->Selection->TopologyTimestamp = SelectionStore->Selection->Mesh->GetTopologyChangeStamp();
-
-		// Don't issue a transaciton if we're cancelling (currently only possible via undo out
-		// of the "Transform" tool). Otherwise, issue a speculative change so that we can clear
-		// the selection if the mesh topology turns out to have changed by the time we run again.
-		if (ShutdownType != EToolShutdownType::Cancel)
-		{
-			UUVSelectToolSpeculativeChangeAPI* SpeculativeChangeAPI = ContextStore->FindContext<UUVSelectToolSpeculativeChangeAPI>();
-			if (!SpeculativeChangeAPI)
-			{
-				SpeculativeChangeAPI = NewObject<UUVSelectToolSpeculativeChangeAPI>();
-				ContextStore->AddContextObject(SpeculativeChangeAPI);
-			}
-			SpeculativeChangeAPI->EmitSpeculativeChange(ChangeRouter, EmitChangeAPI,
-				LOCTEXT("SpeculativeSlectionChangeName", "End Select Tool"));
-		}
-	}
-	else
-	{
-		ContextStore->RemoveContextObjectsOfType<UUVToolMeshSelection>();
+		// (The broadcast here is so that we still broadcast on undo)
+		SelectionMechanic->SetSelection(UMeshSelectionMechanic::FDynamicMeshSelection(), true, true);
 	}
 
 	ChangeRouter->CurrentSelectTool = nullptr;
@@ -496,6 +386,10 @@ void UUVSelectTool::Shutdown(EToolShutdownType ShutdownType)
 	// Calls shutdown on gizmo and destroys it.
 	GetToolManager()->GetPairedGizmoManager()->DestroyAllGizmosByOwner(this);
 
+	ViewportButtonsAPI->OnGizmoModeChange.RemoveAll(this);
+	ViewportButtonsAPI->SetGizmoButtonsEnabled(false);
+
+	ViewportButtonsAPI = nullptr;
 	EmitChangeAPI = nullptr;
 	ChangeRouter = nullptr;
 }
@@ -527,7 +421,9 @@ void UUVSelectTool::UpdateGizmo()
 		TransformGizmo->ReinitializeGizmoTransform(FTransform((FVector)Centroid));
 	}
 
-	TransformGizmo->SetVisibility(bGizmoEnabled && !SelectionMechanic->GetCurrentSelection().IsEmpty());
+	TransformGizmo->SetVisibility(
+		ViewportButtonsAPI->GetGizmoMode() != UUVToolViewportButtonsAPI::EGizmoMode::Select
+		&& !SelectionMechanic->GetCurrentSelection().IsEmpty());
 }
 
 void UUVSelectTool::ConfigureSelectionModeFromControls()
@@ -700,20 +596,6 @@ void UUVSelectTool::UpdateLivePreviewLines()
 	}
 }
 
-
-
-void UUVSelectTool::SetGizmoEnabled(bool bEnabledIn)
-{
-	bGizmoEnabled = bEnabledIn;
-
-	// SetGizmoEnabled may be called before or after Setup, hence the check here to see if
-	// the gizmo is set up.
-	if (TransformGizmo)
-	{
-		UpdateGizmo();
-	}
-}
-
 void UUVSelectTool::GizmoTransformStarted(UTransformProxy* Proxy)
 {
 	bInDrag = true;
diff --git a/Engine/Plugins/Experimental/UVEditor/Source/UVEditorTools/Public/UVSelectTool.h b/Engine/Plugins/Experimental/UVEditor/Source/UVEditorTools/Public/UVSelectTool.h
--- a/Engine/Plugins/Experimental/UVEditor/Source/UVEditorTools/Public/UVSelectTool.h
+++ b/Engine/Plugins/Experimental/UVEditor/Source/UVEditorTools/Public/UVSelectTool.h
@@ -28,8 +28,9 @@ class UUVSelectToolChangeRouter;
 class UUVToolStateObjectStore;
 class UPreviewMesh;
 class UUVEditorToolMeshInput;
-class UUVToolEmitChangeAPI;
 class UUVSeamSewAction;
+class UUVToolEmitChangeAPI;
+class UUVToolViewportButtonsAPI;
 
 UCLASS()
 class UVEDITORTOOLS_API UUVSelectToolBuilder : public UInteractiveToolBuilder
@@ -40,8 +41,6 @@ public:
 	virtual bool CanBuildTool(const FToolBuilderState& SceneState) const override;
 	virtual UInteractiveTool* BuildTool(const FToolBuilderState& SceneState) const override;
 
-	bool bGizmoEnabled = false;
-
 	// This is a pointer so that it can be updated under the builder without
 	// having to set it in the mode after initializing targets.
 	const TArray<TObjectPtr<UUVEditorToolMeshInput>>* Targets = nullptr;
@@ -119,7 +118,6 @@ class UVEDITORTOOLS_API UUVSelectTool : public UInteractiveTool
 
 public:
 	virtual void SetWorld(UWorld* World) { TargetWorld = World; }
-	virtual void SetGizmoEnabled(bool bEnabledIn);
 	
 	/**
 	 * The tool will operate on the meshes given here.
@@ -172,6 +170,9 @@ protected:
 	UPROPERTY()
 	TObjectPtr<USelectToolActionPropertySet> ToolActions = nullptr;
 
+	UPROPERTY()
+	TObjectPtr<UUVToolViewportButtonsAPI> ViewportButtonsAPI = nullptr;
+
 	UPROPERTY()
 	TObjectPtr<UMeshSelectionMechanic> SelectionMechanic = nullptr;
 
@@ -212,9 +213,6 @@ protected:
 	void ApplyAction(ESelectToolAction ActionType);
 	const float LivePreviewHighlightThickness = 2.0;
 	const float LivePreviewHighlightDepthOffset = 0.5;
-
-	// We need this flag so that SetGizmoVisibility can be called before Setup() by the tool builder.
-	bool bGizmoEnabled = false;
 };
 
 /**
@@ -230,38 +228,3 @@ public:
 
 	TWeakObjectPtr<UUVSelectTool> CurrentSelectTool = nullptr;
 };
-
-/**
- * A helper context object that allows us to inject an undo transaction back in time, which gets
- * used to deal with the fact that our stored selection may become invalidated by an intervening
- * tool, and needs to be cleared in an undoable transaction before that tool runs.
- * 
- * NOTE: It seems likely (especially after "Transform" is no longer a separate tool from "Select")
- * that we may just decide not to try to store a selection on tool shutdown. In that case we will
- * delete this class.
- */
-UCLASS()
-class UVEDITORTOOLS_API UUVSelectToolSpeculativeChangeAPI : public UUVToolContextObject
-{
-	GENERATED_BODY()
-public:
-
-	/** 
-	 * Emits a tool-independent change that does nothing unless a subsequent InsertIntoLastSpeculativeChange
-	 * call injects a change.
-	 */
-	void EmitSpeculativeChange(UObject* TargetObject, UUVToolEmitChangeAPI* EmitChangeAPI, const FText& TransactionName);
-
-	bool HasSpeculativeChange()
-	{
-		return ContentOfLastSpeculativeChange.IsValid();
-	}
-
-	/**
-	 * Inserts a change into the place marked by the last EmitSpeculativeChange call.
-	 */
-	void InsertIntoLastSpeculativeChange(TUniquePtr<FToolCommandChange> ChangeToInsert);
-
-protected:
-	TSharedPtr<TUniquePtr<FToolCommandChange>> ContentOfLastSpeculativeChange;
-};
\ No newline at end of file
diff --git a/Engine/Plugins/Experimental/UVEditor/Source/UVEditorTools/Public/UVToolContextObjects.h b/Engine/Plugins/Experimental/UVEditor/Source/UVEditorTools/Public/UVToolContextObjects.h
--- a/Engine/Plugins/Experimental/UVEditor/Source/UVEditorTools/Public/UVToolContextObjects.h
+++ b/Engine/Plugins/Experimental/UVEditor/Source/UVEditorTools/Public/UVToolContextObjects.h
@@ -82,6 +82,10 @@ protected:
 	TObjectPtr<UInteractiveToolManager> ToolManager = nullptr;
 };
 
+/**
+ * Allows tools to interact with the 3d preview viewport, which has a separate
+ * world and input router.
+ */
 UCLASS()
 class UVEDITORTOOLS_API UUVToolLivePreviewAPI : public UUVToolContextObject
 {
@@ -100,15 +104,49 @@ protected:
 	TWeakObjectPtr<UInputRouter> InputRouter;
 };
 
-/** Stores a UV mesh selection */
+/**
+ * Allows tools to interact with buttons in the viewport (currently just gizmo controls)
+ */
 UCLASS()
-class UVEDITORTOOLS_API UUVToolMeshSelection : public UUVToolContextObject
+class UVEDITORTOOLS_API UUVToolViewportButtonsAPI : public UUVToolContextObject
 {
 	GENERATED_BODY()
 public:
 
-	TSharedPtr<UE::Geometry::FDynamicMeshSelection, ESPMode::ThreadSafe> Selection 
-		= MakeShared<UE::Geometry::FDynamicMeshSelection, ESPMode::ThreadSafe>();
+	enum class EGizmoMode
+	{
+		Select,
+		Transform
+	};
+
+	void SetGizmoButtonsEnabled(bool bOn)
+	{
+		bButtonsEnabled = bOn;
+	}
+
+	bool AreGizmoButtonsEnabled()
+	{
+		return bButtonsEnabled;
+	}
+
+	void SetGizmoMode(EGizmoMode ModeIn, bool bBroadcast = true)
+	{
+		GizmoMode = ModeIn;
+		OnGizmoModeChange.Broadcast(GizmoMode);
+	}
+
+	EGizmoMode GetGizmoMode()
+	{
+		return GizmoMode;
+	}
+
+	DECLARE_MULTICAST_DELEGATE_OneParam(FOnGizmoModeChange, EGizmoMode NewGizmoMode);
+	FOnGizmoModeChange OnGizmoModeChange;
+
+protected:
+
+	bool bButtonsEnabled = false;
+	EGizmoMode GizmoMode = EGizmoMode::Select;
 };
 
 /** Stores UV mesh AABB trees */
