UVEditor: Force both viewports to always be realtime, and disable motion blur in them.

#rnx

diff --git a/Engine/Plugins/Experimental/UVEditor/Source/UVEditor/Private/UVEditorToolkit.cpp b/Engine/Plugins/Experimental/UVEditor/Source/UVEditor/Private/UVEditorToolkit.cpp
--- a/Engine/Plugins/Experimental/UVEditor/Source/UVEditor/Private/UVEditorToolkit.cpp
+++ b/Engine/Plugins/Experimental/UVEditor/Source/UVEditor/Private/UVEditorToolkit.cpp
@@ -369,6 +369,17 @@ void FUVEditorToolkit::PostInitAssetEditor()
 
 
 	// Adjust our main (2D) viewport:
+	auto SetCommonViewportClientOptions = [](FEditorViewportClient* Client)
+	{
+		// Normally the bIsRealtime flag is determined by whether the connection is remote, but our
+		// tools require always being ticked.
+		Client->SetRealtime(true);
+
+		// Disable motion blur effects that cause our renders to "fade in" as things are moved
+		Client->EngineShowFlags.SetTemporalAA(false);
+		Client->EngineShowFlags.SetMotionBlur(false);
+	};
+	SetCommonViewportClientOptions(ViewportClient.Get());
 
 	// Ortho has too many problems with rendering things, unfortunately, so we should use perspective.
 	ViewportClient->SetViewportType(ELevelViewportType::LVT_Perspective);
@@ -391,14 +402,13 @@ void FUVEditorToolkit::PostInitAssetEditor()
 	// If exposure isn't set to fixed, it will flash as we stare into the void
 	ViewportClient->ExposureSettings.bFixed = true;
 
-	// TODO: Disable temporal aa or whatever else is blurring the lines as the camera moves.
-
 	// We need the viewport client to start out focused, or else it won't get ticked until
 	// we click inside it.
 	ViewportClient->ReceivedFocus(ViewportClient->Viewport);
 
 
 	// Adjust our live preview (3D) viewport
+	SetCommonViewportClientOptions(LivePreviewViewportClient.Get());
 	// TODO: This should not be hardcoded
 	LivePreviewViewportClient->SetViewLocation(FVector(-200, 100, 100));
 	LivePreviewViewportClient->SetLookAtLocation(FVector(0, 0, 0));
