UVEditor: Add UV Layer selection.

#rnx

diff --git a/Engine/Plugins/Experimental/MeshModelingToolsetExp/Source/ModelingComponents/Private/ModelingToolTargetUtil.cpp b/Engine/Plugins/Experimental/MeshModelingToolsetExp/Source/ModelingComponents/Private/ModelingToolTargetUtil.cpp
--- a/Engine/Plugins/Experimental/MeshModelingToolsetExp/Source/ModelingComponents/Private/ModelingToolTargetUtil.cpp
+++ b/Engine/Plugins/Experimental/MeshModelingToolsetExp/Source/ModelingComponents/Private/ModelingToolTargetUtil.cpp
@@ -9,6 +9,7 @@
 #include "TargetInterfaces/MeshDescriptionCommitter.h"
 #include "TargetInterfaces/MeshDescriptionProvider.h"
 #include "TargetInterfaces/PrimitiveComponentBackedTarget.h"
+#include "TargetInterfaces/AssetBackedTarget.h"
 #include "TargetInterfaces/DynamicMeshSource.h"
 
 #include "ModelingObjectsCreationAPI.h"
@@ -48,6 +49,23 @@ UPrimitiveComponent* UE::ToolTarget::GetTargetComponent(UToolTarget* Target)
 	return nullptr;
 }
 
+FString UE::ToolTarget::GetHumanReadableName(UToolTarget* Target)
+{
+	IAssetBackedTarget* TargetAsset = Cast<IAssetBackedTarget>(Target);
+	if (TargetAsset)
+	{
+		return TargetAsset->GetSourceData()->GetName();
+	}
+
+	IPrimitiveComponentBackedTarget* TargetComponent = Cast<IPrimitiveComponentBackedTarget>(Target);
+	if (TargetComponent)
+	{
+		return TargetComponent->GetOwnerComponent()->GetFullGroupName(false);
+	}
+
+	return FString("");
+}
+
 bool UE::ToolTarget::HideSourceObject(UToolTarget* Target)
 {
 	IPrimitiveComponentBackedTarget* TargetComponent = Cast<IPrimitiveComponentBackedTarget>(Target);
@@ -320,6 +338,16 @@ UE::ToolTarget::EDynamicMeshUpdateResult UE::ToolTarget::CommitDynamicMeshUVUpda
 		return EDynamicMeshUpdateResult::Ok;
 	}
 
+	IDynamicMeshCommitter* DynamicMeshCommitter = Cast<IDynamicMeshCommitter>(Target);
+	if (DynamicMeshCommitter)
+	{
+		IDynamicMeshCommitter::FDynamicMeshCommitInfo CommitInfo(false);
+		CommitInfo.bUVsChanged = true;
+
+		DynamicMeshCommitter->CommitDynamicMesh(*UpdatedMesh, CommitInfo);
+
+		return EDynamicMeshUpdateResult::Ok;
+	}
 
 	IMeshDescriptionCommitter* MeshDescriptionCommitter = Cast<IMeshDescriptionCommitter>(Target);
 	if (!ensure(MeshDescriptionCommitter))
diff --git a/Engine/Plugins/Experimental/MeshModelingToolsetExp/Source/ModelingComponents/Public/ModelingToolTargetUtil.h b/Engine/Plugins/Experimental/MeshModelingToolsetExp/Source/ModelingComponents/Public/ModelingToolTargetUtil.h
--- a/Engine/Plugins/Experimental/MeshModelingToolsetExp/Source/ModelingComponents/Public/ModelingToolTargetUtil.h
+++ b/Engine/Plugins/Experimental/MeshModelingToolsetExp/Source/ModelingComponents/Public/ModelingToolTargetUtil.h
@@ -37,6 +37,11 @@ MODELINGCOMPONENTS_API AActor* GetTargetActor(UToolTarget* Target);
  */
 MODELINGCOMPONENTS_API UPrimitiveComponent* GetTargetComponent(UToolTarget* Target);
 
+/**
+ * @return a human readable identifier for a ToolTarget, based on the underlying nature of it.
+ */
+MODELINGCOMPONENTS_API FString GetHumanReadableName(UToolTarget* Target);
+
 /**
  * Hide the "Source Object" (eg PrimitiveComponent, Actor, etc) backing a ToolTarget
  * @return true on success
diff --git a/Engine/Plugins/Experimental/UVEditor/Source/UVEditor/Private/UVEditorMode.cpp b/Engine/Plugins/Experimental/UVEditor/Source/UVEditor/Private/UVEditorMode.cpp
--- a/Engine/Plugins/Experimental/UVEditor/Source/UVEditor/Private/UVEditorMode.cpp
+++ b/Engine/Plugins/Experimental/UVEditor/Source/UVEditor/Private/UVEditorMode.cpp
@@ -9,11 +9,12 @@
 #include "Framework/Commands/UICommandList.h"
 #include "InteractiveTool.h"
 #include "MeshOpPreviewHelpers.h" // UMeshOpPreviewWithBackgroundCompute
+#include "ModelingToolTargetUtil.h"
 #include "PreviewMesh.h"
 #include "TargetInterfaces/AssetBackedTarget.h"
 #include "TargetInterfaces/MaterialProvider.h"
-#include "TargetInterfaces/DynamicMeshProvider.h"
-#include "TargetInterfaces/DynamicMeshCommitter.h"
+#include "TargetInterfaces/MeshDescriptionCommitter.h"
+#include "TargetInterfaces/MeshDescriptionProvider.h"
 #include "ToolSetupUtil.h"
 #include "ToolTargets/UVEditorToolMeshInput.h"
 #include "ToolTargetManager.h"
@@ -27,6 +28,7 @@
 #include "UVEditorToolUtil.h"
 #include "UVToolContextObjects.h"
 #include "UVEditorBackgroundPreview.h"
+#include "UVEditorModeChannelProperties.h"
 #include "Editor.h"
 
 #define LOCTEXT_NAMESPACE "UUVEditorMode"
@@ -37,11 +39,60 @@ const FEditorModeID UUVEditorMode::EM_UVEditorModeId = TEXT("EM_UVEditorMode");
 
 namespace UVEditorModeLocals
 {
-	// TODO: This is temporary. We need to be able to configure the layer we look at,
-	// and ideally show multiple layers at a time.
-	const int32 UVLayerIndex = 0;	
+	// The layer we open when we first open the UV editor
+	const int32 DefaultUVLayerIndex = 0;
 
+	// Determines what the default tool is
 	const FString DefaultToolIdentifier = TEXT("UVSelectTool");
+
+	const FText UVLayerChangeTransactionName = LOCTEXT("UVLayerChangeTransactionName", "Change UV Layer");
+
+	/** 
+	 * Change for undoing/redoing displayed layer changes.
+	 */
+	class FInputObjectUVLayerChange : public FToolCommandChange
+	{
+	public:
+		FInputObjectUVLayerChange(int32 AssetIDIn, int32 OldUVLayerIndexIn, int32 NewUVLayerIndexIn)
+			: AssetID(AssetIDIn)
+			, OldUVLayerIndex(OldUVLayerIndexIn)
+			, NewUVLayerIndex(NewUVLayerIndexIn)
+		{
+		}
+
+		virtual void Apply(UObject* Object) override
+		{
+			UUVEditorMode* UVEditorMode = Cast<UUVEditorMode>(Object);
+			UVEditorMode->ChangeInputObjectLayer(AssetID, NewUVLayerIndex);
+			UVEditorMode->UpdateSelectedLayer();
+		}
+
+		virtual void Revert(UObject* Object) override
+		{
+			UUVEditorMode* UVEditorMode = Cast<UUVEditorMode>(Object);
+			UVEditorMode->ChangeInputObjectLayer(AssetID, OldUVLayerIndex);
+			UVEditorMode->UpdateSelectedLayer();
+		}
+
+		virtual bool HasExpired(UObject* Object) const override
+		{
+			UUVEditorMode* UVEditorMode = Cast<UUVEditorMode>(Object);
+			return !(UVEditorMode && UVEditorMode->IsActive());
+		}
+
+		virtual FString ToString() const override
+		{
+			return TEXT("UVEditorModeLocals::FInputObjectUVLayerChange");
+		}
+
+	protected:
+		int32 AssetID;
+		int32 OldUVLayerIndex;
+		int32 NewUVLayerIndex;
+	};
+
+
+
 }
 
 const FToolTargetTypeRequirements& UUVEditorMode::GetToolTargetRequirements()
@@ -49,9 +100,12 @@ const FToolTargetTypeRequirements& UUVEditorMode::GetToolTargetRequirements()
 	static const FToolTargetTypeRequirements ToolTargetRequirements =
 		FToolTargetTypeRequirements({
 			UMaterialProvider::StaticClass(),
-			UDynamicMeshCommitter::StaticClass(),
-			UDynamicMeshProvider::StaticClass(),
-			UAssetBackedTarget::StaticClass()
+
+			// What we actually care about is dynamic meshes, but we don't currently have
+			// a standardized dynamic mesh commiter/provider interface, because UDynamicMesh
+			// doesn't implement IDynamicMeshCommitter
+			UMeshDescriptionCommitter::StaticClass(),
+			UMeshDescriptionProvider::StaticClass()
 			});
 	return ToolTargetRequirements;
 }
@@ -71,14 +125,24 @@ void UUVEditorMode::Enter()
 
 	BackgroundVisualization = NewObject<UUVEditorBackgroundPreview>(this);
 	BackgroundVisualization->CreateInWorld(GetWorld(), FTransform::Identity);
-	EditorModeToolkit->SetBackgroundSettings(BackgroundVisualization->Settings);
 
-	BackgroundVisualization->Settings->WatchProperty(BackgroundVisualization->Settings->bVisible, [this](bool IsVisible) { UpdateTriangleMaterialBasedOnBackground(IsVisible); });
-	PropertyObjectsToTick.Add(BackgroundVisualization->Settings);
+	UVChannelProperties = NewObject<UUVEditorUVChannelProperties>(this);
+	UVChannelProperties->WatchProperty(UVChannelProperties->Asset, [this](FString UVAsset) {SwitchActiveAsset(UVAsset); });
+	UVChannelProperties->WatchProperty(UVChannelProperties->UVChannel, [this](FString UVChannel) {SwitchActiveChannel(UVChannel); });
+
+	BackgroundVisualization->Settings->WatchProperty(BackgroundVisualization->Settings->bVisible, 
+		[this](bool IsVisible) { UpdateTriangleMaterialBasedOnBackground(IsVisible); });
+
+	AddDisplayedPropertySet(UVChannelProperties);
+	AddDisplayedPropertySet(BackgroundVisualization->Settings);
+		
+	StaticCastSharedPtr<FUVEditorModeToolkit>(Toolkit)->SetModeDetailsViewObjects(PropertyObjectsToDisplay);
 
 	RegisterTools();
 
 	ActivateDefaultTool();
+
+	bIsActive = true;
 }
 
 void UUVEditorMode::RegisterTools()
@@ -98,9 +162,15 @@ void UUVEditorMode::RegisterTools()
 	RegisterTool(CommandInfos.BeginParameterizeMeshTool, TEXT("UVParameterizeMeshTool"), UVEditorParameterizeMeshToolBuilder);
 }
 
+void UUVEditorMode::AddDisplayedPropertySet(const TObjectPtr<UInteractiveToolPropertySet>& PropertySet)
+{
+	PropertyObjectsToTick.Add(PropertySet);
+	PropertyObjectsToDisplay.Add(PropertySet);
+}
+
 void UUVEditorMode::CreateToolkit()
 {
-	Toolkit = EditorModeToolkit = MakeShared<FUVEditorModeToolkit>();
+	Toolkit = MakeShared<FUVEditorModeToolkit>();
 }
 
 void UUVEditorMode::ActivateDefaultTool()
@@ -155,34 +225,77 @@ void UUVEditorMode::Exit()
 		ToolInput->Shutdown();
 	}
 	ToolInputObjects.Reset();
+	WireframesToTick.Reset();
+	OriginalObjectsToEdit.Reset();
+	
+	for (TObjectPtr<UMeshOpPreviewWithBackgroundCompute> Preview : AppliedPreviews)
+	{
+		Preview->Shutdown();
+	}
+	AppliedPreviews.Reset();
+	AppliedCanonicalMeshes.Reset();
+	ToolTargets.Reset();
 
 	if (BackgroundVisualization)
 	{
 		BackgroundVisualization->Disconnect();
+		BackgroundVisualization = nullptr;
 	}
 
+	UVChannelProperties = nullptr;
+	PropertyObjectsToTick.Empty();
+	PropertyObjectsToDisplay.Empty();
+	LivePreviewWorld = nullptr;
+
+	bIsActive = false;
+
 	Super::Exit();
 }
 
-void UUVEditorMode::InitializeTargets(TArray<TObjectPtr<UObject>>& AssetsIn, TArray<FTransform>* TransformsIn)
+void UUVEditorMode::InitializeTargets(const TArray<TObjectPtr<UObject>>& AssetsIn, 
+	const TArray<FTransform>& TransformsIn)
 {
 	using namespace UVEditorModeLocals;
 
 	OriginalObjectsToEdit = AssetsIn;
+	Transforms = TransformsIn;
 
 	// Build the tool targets that provide us with 3d dynamic meshes
 	UUVEditorSubsystem* UVSubsystem = GEditor->GetEditorSubsystem<UUVEditorSubsystem>();
 	UVSubsystem->BuildTargets(AssetsIn, GetToolTargetRequirements(), ToolTargets);
 
-	// Get an api for manipulating things in the live preview, which has its own world and input router
+	// For creating the actual input objects, we'll need pointers both to the 2d unwrap world and the
+	// 3d preview world. We already have the 2d world in GetWorld(). Get the 3d one.
 	UContextObjectStore* ContextStore = GetInteractiveToolsContext()->ToolManager->GetContextObjectStore();
 	UUVToolLivePreviewAPI* LivePreviewAPI = ContextStore->FindContext<UUVToolLivePreviewAPI>();
 	check(LivePreviewAPI);
-	UWorld* LivePreviewWorld = LivePreviewAPI->GetLivePreviewWorld();
+	LivePreviewWorld = LivePreviewAPI->GetLivePreviewWorld();
 
-	double ScaleFactor = GetUVMeshScalingFactor();
+	// Collect the 3d dynamic meshes from targets. There will always be one for each asset, and the AssetID
+	// of each asset will be the index into these arrays. Individual input objects (representing a specific
+	// UV layer), will point to these existing 3d meshes.
+	for (UToolTarget* Target : ToolTargets)
+	{
+		// The applied canonical mesh is the 3d mesh with all of the layer changes applied. If we switch
+		// to a different layer, the changes persist in the applied canonical.
+		TSharedPtr<FDynamicMesh3> AppliedCanonical = MakeShared<FDynamicMesh3>(UE::ToolTarget::GetDynamicMeshCopy(Target));
+		AppliedCanonicalMeshes.Add(AppliedCanonical);
+
+		// Make a preview version of the applied canonical to show. Tools can attach computes to this, though
+		// they would have to take care if we ever allow multiple layers to be displayed for one asset, to
+		// avoid trying to attach two computes to the same preview object (in which case one would be thrown out)
+		UMeshOpPreviewWithBackgroundCompute* AppliedPreview = NewObject<UMeshOpPreviewWithBackgroundCompute>();
+		AppliedPreview->Setup(LivePreviewWorld);
+		AppliedPreview->PreviewMesh->UpdatePreview(AppliedCanonical.Get());
+
+		FComponentMaterialSet MaterialSet = UE::ToolTarget::GetMaterialSet(Target);
+		AppliedPreview->ConfigureMaterials(MaterialSet.Materials, 
+			ToolSetupUtil::GetDefaultWorkingMaterial(GetToolManager()));
+		AppliedPreviews.Add(AppliedPreview);
+	}
 
-	// These functions will determine the mapping between UV values and the resulting mesh vertex positions. 
+	// When creating UV unwraps, these functions will determine the mapping between UV values and the
+	// resulting unwrap mesh vertex positions. 
 	// If we're looking down on the unwrapped mesh, with the Z axis towards us, we want U's to be right, and
 	// V's to be up. In Unreal's left-handed coordinate system, this means that we map U's to world Y
 	// and V's to world X.
@@ -192,6 +305,7 @@ void UUVEditorMode::InitializeTargets(TArray<TObjectPtr<UObject>>& AssetsIn, TAr
 	// frequently end up negative).
 	// The ScaleFactor just scales the mesh up. Scaling the mesh up makes it easier to zoom in
 	// further into the display before getting issues with the camera near plane distance.
+	double ScaleFactor = GetUVMeshScalingFactor();
 	auto UVToVertPosition = [this, ScaleFactor](const FVector2f& UV)
 	{
 		return FVector3d((1 - UV.Y) * ScaleFactor, UV.X * ScaleFactor, 0);
@@ -201,18 +315,22 @@ void UUVEditorMode::InitializeTargets(TArray<TObjectPtr<UObject>>& AssetsIn, TAr
 		return FVector2D(VertPosition.Y / ScaleFactor, 1 - (VertPosition.X / ScaleFactor));
 	};
 
-	// For each of our tool targets, we're going to create a UV tool input object that bundles together
-	// the 3d preview and the unwrapped UV layer, with both a reference and a "preview with background op"
-	// versions of both
-	for (UToolTarget* Target : ToolTargets)
+	// Construct the full input objects that the tools actually operate on.
+	for (int32 AssetID = 0; AssetID < ToolTargets.Num(); ++AssetID)
 	{
 		UUVEditorToolMeshInput* ToolInputObject = NewObject<UUVEditorToolMeshInput>();
-		
-		if (!ToolInputObject->InitializeMeshes(Target, UVLayerIndex,
+
+		if (!ToolInputObject->InitializeMeshes(ToolTargets[AssetID], AppliedCanonicalMeshes[AssetID],
+			AppliedPreviews[AssetID], AssetID, DefaultUVLayerIndex,
 			GetWorld(), LivePreviewWorld, ToolSetupUtil::GetDefaultWorkingMaterial(GetToolManager()),
 			UVToVertPosition, VertPositionToUV))
 		{
-			continue;
+			return;
+		}
+
+		if (Transforms.Num() == ToolTargets.Num())
+		{
+			ToolInputObject->AppliedPreview->PreviewMesh->SetTransform(Transforms[AssetID]);
 		}
 
 		ToolInputObject->UnwrapPreview->PreviewMesh->SetMaterial(
@@ -222,9 +340,6 @@ void UUVEditorMode::InitializeTargets(TArray<TObjectPtr<UObject>>& AssetsIn, TAr
 				TriangleDepthOffset,
 				TriangleOpacity));
 
-		// Initialize our timestamp so we can later detect changes
-		MeshChangeStamps.Add(ToolInputObject->UnwrapCanonical->GetShapeChangeStamp());
-
 		// Set up the wireframe display of the unwrapped mesh.
 		UMeshElementsVisualizer* WireframeDisplay = NewObject<UMeshElementsVisualizer>(this);
 		WireframeDisplay->CreateInWorld(GetWorld(), FTransform::Identity);
@@ -254,16 +369,16 @@ void UUVEditorMode::InitializeTargets(TArray<TObjectPtr<UObject>>& AssetsIn, TAr
 		// The tool input object will hold on to the wireframe for the purposes of updating it and cleaning it up
 		ToolInputObject->WireframeDisplay = WireframeDisplay;
 
+		// Initialize our changestamp so we can later detect changes
+		LastSeenChangeStamps.Add(ToolInputObject->UnwrapCanonical->GetShapeChangeStamp());
+
 		ToolInputObjects.Add(ToolInputObject);
 	}
-	
-	if (TransformsIn && TransformsIn->Num() == AssetsIn.Num())
-	{
-		for (int i = 0; i < ToolInputObjects.Num(); ++i)
-		{
-			ToolInputObjects[i]->AppliedPreview->PreviewMesh->SetTransform((*TransformsIn)[i]);
-		}
-	}
+
+	// Initialize our layer selector
+	UVChannelProperties->Initialize(ToolTargets, AppliedCanonicalMeshes, true);
+	PreviousUVLayerIndex = UVChannelProperties->GetSelectedChannelIndex();
+	PendingUVLayerIndex = PreviousUVLayerIndex;
 }
 
 void UUVEditorMode::EmitToolIndependentObjectChange(UObject* TargetObject, TUniquePtr<FToolCommandChange> Change, const FText& Description)
@@ -273,24 +388,15 @@ void UUVEditorMode::EmitToolIndependentObjectChange(UObject* TargetObject, TUniq
 
 bool UUVEditorMode::HaveUnappliedChanges()
 {
-	for (int32 i = 0; i < ToolInputObjects.Num(); ++i)
-	{
-		if (ToolInputObjects[i]->UnwrapCanonical->GetShapeChangeStamp() != MeshChangeStamps[i])
-		{
-			return true;
-		}
-	}
-	return false;
+	return ModifiedAssetIDs.Num() > 0;
 }
 
 void UUVEditorMode::GetAssetsWithUnappliedChanges(TArray<TObjectPtr<UObject>> UnappliedAssetsOut)
 {
-	for (int32 i = 0; i < ToolInputObjects.Num(); ++i)
+	for (int32 AssetID : ModifiedAssetIDs)
 	{
-		if (ToolInputObjects[i]->UnwrapCanonical->GetShapeChangeStamp() != MeshChangeStamps[i])
-		{
-			UnappliedAssetsOut.Add(OriginalObjectsToEdit[i]);
-		}
+		// The asset ID corresponds to the index into OriginalObjectsToEdit
+		UnappliedAssetsOut.Add(OriginalObjectsToEdit[AssetID]);
 	}
 }
 
@@ -300,18 +406,14 @@ void UUVEditorMode::ApplyChanges()
 
 	GetToolManager()->BeginUndoTransaction(LOCTEXT("UVEditorApplyChangesTransaction", "UV Editor Apply Changes"));
 
-	IDynamicMeshCommitter::FDynamicMeshCommitInfo CommitInfo(false);
-	CommitInfo.bUVsChanged = true;
-
-	for (int32 i = 0; i < ToolInputObjects.Num(); ++i)
+	for (int32 AssetID : ModifiedAssetIDs)
 	{
-		if (ToolInputObjects[i]->UnwrapCanonical->GetShapeChangeStamp() != MeshChangeStamps[i])
-		{
-			Cast<IDynamicMeshCommitter>(ToolTargets[i])->CommitDynamicMesh(*ToolInputObjects[i]->AppliedCanonical, CommitInfo);
-			MeshChangeStamps[i] = ToolInputObjects[i]->UnwrapCanonical->GetShapeChangeStamp();
-		}
+		// The asset ID corresponds to the index into ToolTargets and AppliedCanonicalMeshes
+		UE::ToolTarget::CommitDynamicMeshUVUpdate(ToolTargets[AssetID], AppliedCanonicalMeshes[AssetID].Get());
 	}
 
+	ModifiedAssetIDs.Reset();
+
 	GetToolManager()->EndUndoTransaction();
 }
 
@@ -344,8 +446,36 @@ void UUVEditorMode::UpdateTriangleMaterialBasedOnBackground(bool IsBackgroundVis
 
 void UUVEditorMode::ModeTick(float DeltaTime)
 {
+	using namespace UVEditorModeLocals;
+
 	Super::ModeTick(DeltaTime);
 
+	if (PreviousUVLayerIndex!= PendingUVLayerIndex)
+	{
+		int32 AssetID = UVChannelProperties->GetSelectedAssetID();
+		if (ensure(ToolInputObjects[AssetID]->UVLayerIndex != PendingUVLayerIndex))
+		{
+			GetToolManager()->BeginUndoTransaction(UVLayerChangeTransactionName);
+
+			// TODO: Perhaps we need our own interactive tools context that allows this kind of "end tool now"
+			// call. We can't do the normal EndTool call because we cannot defer shutdown until the tick.
+			GetInteractiveToolsContext()->ToolManager->DeactivateTool(EToolSide::Mouse, EToolShutdownType::Cancel);
+
+			ChangeInputObjectLayer(AssetID, PendingUVLayerIndex);
+
+			GetInteractiveToolsContext()->GetTransactionAPI()->AppendChange(this,
+				MakeUnique<FInputObjectUVLayerChange>(AssetID, PreviousUVLayerIndex, PendingUVLayerIndex), 
+				UVLayerChangeTransactionName);
+
+			GetToolManager()->EndUndoTransaction();
+
+			ActivateDefaultTool();
+			//GetInteractiveToolsContext()->ToolManager->ActivateTool(EToolSide::Mouse);
+
+		}
+		PreviousUVLayerIndex = PendingUVLayerIndex;
+	}
+
 	for (TObjectPtr<UInteractiveToolPropertySet>& Propset : PropertyObjectsToTick)
 	{
 		if (Propset)
@@ -376,9 +506,101 @@ void UUVEditorMode::ModeTick(float DeltaTime)
 
 	for (int i = 0; i < ToolInputObjects.Num(); ++i)
 	{
-		ToolInputObjects[i]->AppliedPreview->Tick(DeltaTime);
-		ToolInputObjects[i]->UnwrapPreview->Tick(DeltaTime);
+		TObjectPtr<UUVEditorToolMeshInput> ToolInput = ToolInputObjects[i];
+		ToolInput->AppliedPreview->Tick(DeltaTime);
+		ToolInput->UnwrapPreview->Tick(DeltaTime);
+
+		int32 ChangeStamp = ToolInput->UnwrapCanonical->GetShapeChangeStamp();
+		if (LastSeenChangeStamps[i] != ChangeStamp)
+		{
+			LastSeenChangeStamps[i] = ChangeStamp;
+			ModifiedAssetIDs.Add(ToolInput->AssetID);
+		}
+	}
+}
+
+
+void UUVEditorMode::SwitchActiveAsset(const FString& UVAsset)
+{
+	if (UVChannelProperties)
+	{
+		// Not doing an ensure here because the "revert to default" can give us an empty string
+		UVChannelProperties->ValidateUVAssetSelection(true);
+		UpdateSelectedLayer();
+	}
+}
+
+void UUVEditorMode::UpdateSelectedLayer()
+{
+	int32 AssetID = UVChannelProperties->GetSelectedAssetID();
+	if (!ensure(AssetID != INDEX_NONE))
+	{
+		return;
+	}
+	const TArray<FString>& ChannelNames = UVChannelProperties->GetUVChannelNames();
+
+	// Find the currently selected layer for the current asset
+	bool bFound = false;
+	for (auto InputObject : ToolInputObjects)
+	{
+		if (InputObject->AssetID == AssetID)
+		{
+			bFound = true;
+
+			UVChannelProperties->UVChannel = ChannelNames[InputObject->UVLayerIndex];
+			UVChannelProperties->SilentUpdateWatched();
+			PreviousUVLayerIndex = InputObject->UVLayerIndex;
+			break;
+		}
+	}
+	if (!ensure(bFound))
+	{
+		UVChannelProperties->UVChannel = TEXT("");
+		PreviousUVLayerIndex = INDEX_NONE;
 	}
+
+	PendingUVLayerIndex = PreviousUVLayerIndex;
+}
+
+
+void UUVEditorMode::SwitchActiveChannel(const FString& UVChannel)
+{
+	using namespace UVEditorModeLocals;
+
+	if (UVChannelProperties)
+	{
+		// Not doing an ensure because the "revert to default" can give us an empty string
+		UVChannelProperties->ValidateUVChannelSelection(true);
+		int32 NewUVLayerIndex = UVChannelProperties->GetSelectedChannelIndex();
+
+		if (ensure(NewUVLayerIndex != IndexConstants::InvalidID))
+		{
+			PendingUVLayerIndex = NewUVLayerIndex;
+		}
+	}
+}
+
+void UUVEditorMode::ChangeInputObjectLayer(int32 AssetID, int32 NewLayerIndex)
+{
+	using namespace UVEditorModeLocals;
+
+	bool bFound = false;
+	for (auto InputObject : ToolInputObjects)
+	{
+		if (InputObject->AssetID == AssetID)
+		{
+			bFound = true;
+
+			if (ensure(InputObject->UVLayerIndex != NewLayerIndex))
+			{
+				InputObject->UVLayerIndex = NewLayerIndex;
+				InputObject->UpdateAllFromAppliedCanonical();
+			}
+			break;
+		}
+	}
+
+	ensure(bFound);
 }
 
 #undef LOCTEXT_NAMESPACE
diff --git a/Engine/Plugins/Experimental/UVEditor/Source/UVEditor/Private/UVEditorModeChannelProperties.cpp b/Engine/Plugins/Experimental/UVEditor/Source/UVEditor/Private/UVEditorModeChannelProperties.cpp
new file mode 100644
--- /dev/null
+++ b/Engine/Plugins/Experimental/UVEditor/Source/UVEditor/Private/UVEditorModeChannelProperties.cpp
@@ -0,0 +1,139 @@
+// Copyright Epic Games, Inc. All Rights Reserved.
+
+#include "UVEditorModeChannelProperties.h"
+
+#include "DynamicMesh/DynamicMesh3.h"
+#include "DynamicMesh/DynamicMeshAttributeSet.h"
+#include "ModelingToolTargetUtil.h"
+#include "ToolTargets/UVEditorToolMeshInput.h"
+
+void UUVEditorUVChannelProperties::Initialize(
+	TArray<TObjectPtr<UToolTarget>>& TargetObjects, 
+	TArray<TSharedPtr<UE::Geometry::FDynamicMesh3>> AppliedCanonicalMeshes,
+	bool bInitializeSelection)
+{
+	if (!ensure(TargetObjects.Num() == AppliedCanonicalMeshes.Num()))
+	{
+		return;
+	}
+
+	UVAssetNames.Reset(TargetObjects.Num());
+	NumUVLayersPerAsset.Reset(TargetObjects.Num());
+
+	for (int i = 0; i < TargetObjects.Num(); ++i)
+	{
+		UVAssetNames.Add(UE::ToolTarget::GetHumanReadableName(TargetObjects[i]));
+		
+		int32 NumUVChannels = AppliedCanonicalMeshes[i]->HasAttributes() ?
+			AppliedCanonicalMeshes[i]->Attributes()->NumUVLayers() : 0;
+		
+		NumUVLayersPerAsset.Add(NumUVChannels);
+	}
+
+	UVChannelNames.Reset();
+	if (bInitializeSelection)
+	{
+		Asset = UVAssetNames[0];
+		GetUVChannelNames();
+		UVChannel = UVChannelNames.Num() > 0 ? UVChannelNames[0] : TEXT("");
+	}
+}
+
+const TArray<FString>& UUVEditorUVChannelProperties::GetAssetNames()
+{
+	return UVAssetNames;
+}
+
+const TArray<FString>& UUVEditorUVChannelProperties::GetUVChannelNames()
+{
+	int32 FoundAssetIndex = UVAssetNames.IndexOfByKey(Asset);
+
+	if (FoundAssetIndex == INDEX_NONE)
+	{
+		UVChannelNames.Reset();
+		return UVChannelNames;
+	}
+
+	if (UVChannelNames.Num() != NumUVLayersPerAsset[FoundAssetIndex])
+	{
+		UVChannelNames.Reset();
+		for (int32 i = 0; i < NumUVLayersPerAsset[FoundAssetIndex]; ++i)
+		{
+			UVChannelNames.Add(FString::Printf(TEXT("UV%d"), i));
+		}
+	}
+
+	return UVChannelNames;
+}
+
+bool UUVEditorUVChannelProperties::ValidateUVAssetSelection(bool bUpdateIfInvalid)
+{
+	int32 FoundIndex = UVAssetNames.IndexOfByKey(Asset);
+	if (FoundIndex == INDEX_NONE)
+	{
+		if (bUpdateIfInvalid)
+		{
+			Asset = (UVAssetNames.Num() > 0) ? UVAssetNames[0] : TEXT("");
+		}
+		return false;
+	}
+	return true;
+}
+
+bool UUVEditorUVChannelProperties::ValidateUVChannelSelection(bool bUpdateIfInvalid)
+{
+	bool bValid = ValidateUVAssetSelection(bUpdateIfInvalid);
+
+	int32 FoundAssetIndex = UVAssetNames.IndexOfByKey(Asset);
+	if (FoundAssetIndex == INDEX_NONE)
+	{
+		if (bUpdateIfInvalid)
+		{
+			UVChannel = TEXT("");
+		}
+		return false;
+	}
+
+	int32 FoundIndex = UVChannelNames.IndexOfByKey(UVChannel);
+	if (FoundIndex >= NumUVLayersPerAsset[FoundAssetIndex])
+	{
+		FoundIndex = INDEX_NONE;
+	}
+
+	if (FoundIndex == INDEX_NONE)
+	{
+		if (bUpdateIfInvalid)
+		{
+			UVChannel = (NumUVLayersPerAsset[FoundAssetIndex] > 0) ? UVChannelNames[0] : TEXT("");
+		}
+		return false;
+	}
+
+	return bValid;
+}
+
+int32 UUVEditorUVChannelProperties::GetSelectedAssetID()
+{
+	int32 FoundIndex = UVAssetNames.IndexOfByKey(Asset);
+	if (FoundIndex == INDEX_NONE)
+	{
+		return INDEX_NONE;
+	}
+	return FoundIndex;
+}
+
+int32 UUVEditorUVChannelProperties::GetSelectedChannelIndex(bool bForceToZeroOnFailure)
+{
+	int32 FoundAssetIndex = UVAssetNames.IndexOfByKey(Asset);
+	if (FoundAssetIndex == INDEX_NONE)
+	{
+		return INDEX_NONE;
+	}
+
+	int32 FoundUVIndex = UVChannelNames.IndexOfByKey(UVChannel);
+	if (!ensure(FoundUVIndex < NumUVLayersPerAsset[FoundAssetIndex]))
+	{
+		return INDEX_NONE;
+	}
+	return FoundUVIndex;
+}
\ No newline at end of file
diff --git a/Engine/Plugins/Experimental/UVEditor/Source/UVEditor/Private/UVEditorModeToolkit.cpp b/Engine/Plugins/Experimental/UVEditor/Source/UVEditor/Private/UVEditorModeToolkit.cpp
--- a/Engine/Plugins/Experimental/UVEditor/Source/UVEditor/Private/UVEditorModeToolkit.cpp
+++ b/Engine/Plugins/Experimental/UVEditor/Source/UVEditor/Private/UVEditorModeToolkit.cpp
@@ -50,7 +50,10 @@ namespace UVEditorModeToolkitLocals
 		{
 			// To not be expired, we must be in some non-default tool.
 			UUVEditorMode* Mode = Cast<UUVEditorMode>(Object);
-			return !(Mode && Mode->GetInteractiveToolsContext()->ToolManager->HasAnyActiveTool() && !Mode->IsDefaultToolActive());
+			return !(Mode && Mode->GetInteractiveToolsContext() 
+				&& Mode->GetInteractiveToolsContext()->ToolManager
+				&& Mode->GetInteractiveToolsContext()->ToolManager->HasAnyActiveTool() 
+				&& !Mode->IsDefaultToolActive());
 		}
 
 		virtual FString ToString() const override
@@ -313,13 +316,12 @@ FText FUVEditorModeToolkit::GetBaseToolkitName() const
 	return NSLOCTEXT("UVEditorModeToolkit", "DisplayName", "UVEditorMode");
 }
 
-
-void FUVEditorModeToolkit::SetBackgroundSettings(UObject* InSettingsObject)
+void FUVEditorModeToolkit::SetModeDetailsViewObjects(const TArray<TObjectPtr<UInteractiveToolPropertySet>>& InObjects)
 {
-	BackgroundDetailsView->SetObject(InSettingsObject);
+	TArray<TObjectPtr<UObject>> Objects(InObjects);
+	ModeDetailsView->SetObjects(Objects);
 }
 
-
 void FUVEditorModeToolkit::UpdateActiveToolProperties()
 {
 	UInteractiveTool* CurTool = GetScriptableEditorMode()->GetToolManager()->GetActiveTool(EToolSide::Left);
diff --git a/Engine/Plugins/Experimental/UVEditor/Source/UVEditor/Private/UVEditorToolkit.cpp b/Engine/Plugins/Experimental/UVEditor/Source/UVEditor/Private/UVEditorToolkit.cpp
--- a/Engine/Plugins/Experimental/UVEditor/Source/UVEditor/Private/UVEditorToolkit.cpp
+++ b/Engine/Plugins/Experimental/UVEditor/Source/UVEditor/Private/UVEditorToolkit.cpp
@@ -340,7 +340,7 @@ void FUVEditorToolkit::PostInitAssetEditor()
 	TArray<FTransform> ObjectTransforms;
 	ObjectTransforms.SetNum(ObjectsToEdit.Num());
 
-	UVMode->InitializeTargets(ObjectsToEdit, &ObjectTransforms);
+	UVMode->InitializeTargets(ObjectsToEdit, ObjectTransforms);
 
 	// Plug in the mode tool panel
 	TSharedPtr<FModeToolkit> UVModeToolkit = UVMode->GetToolkit().Pin();
diff --git a/Engine/Plugins/Experimental/UVEditor/Source/UVEditor/Public/UVEditorMode.h b/Engine/Plugins/Experimental/UVEditor/Source/UVEditor/Public/UVEditorMode.h
--- a/Engine/Plugins/Experimental/UVEditor/Source/UVEditor/Public/UVEditorMode.h
+++ b/Engine/Plugins/Experimental/UVEditor/Source/UVEditor/Public/UVEditorMode.h
@@ -6,9 +6,12 @@
 
 #include "Tools/UEdMode.h"
 #include "ToolTargets/ToolTarget.h"
+#include "GeometryBase.h"
 
 #include "UVEditorMode.generated.h"
 
+PREDECLARE_GEOMETRY(class FDynamicMesh3);
+
 class FToolCommandChange;
 class IUVUnwrapDynamicMesh;
 class UMeshElementsVisualizer;
@@ -16,10 +19,13 @@ class UPreviewMesh;
 class UToolTarget;
 class UWorld;
 class FUVEditorModeToolkit;
-class UInteractiveToolPropertySet;
+class UInteractiveToolPropertySet; 
+class UMeshOpPreviewWithBackgroundCompute;
 class UUVToolStateObjectStore;
 class UUVEditorToolMeshInput;
 class UUVEditorBackgroundPreview;
+class UUVEditorUVChannelProperties;
+
 
 /**
  * The UV editor mode is the mode used in the UV asset editor. It holds most of the inter-tool state.
@@ -48,7 +54,13 @@ public:
 	 */
 	static double GetUVMeshScalingFactor() { return 1000; }
 
-	void InitializeTargets(TArray<TObjectPtr<UObject>>& AssetsIn, TArray<FTransform>* TransformsIn);
+	void InitializeTargets(const TArray<TObjectPtr<UObject>>& AssetsIn, const TArray<FTransform>& TransformsIn);
+
+	// public for use by undo/redo
+	void ChangeInputObjectLayer(int32 AssetID, int32 NewLayerIndex);
+	void UpdateSelectedLayer();
+
+	bool IsActive() { return bIsActive; }
 
 	// TODO: We'll probably eventually want a function like this so we can figure out how big our 3d preview
 	// scene is and how we should position the camera intially...
@@ -75,75 +87,127 @@ public:
 	// it to later get (accidentally) more entangled with mode internals. At the same time, we're not
 	// sure whether we want to make the UEdMode one public. So this is the minimal-impact tweak.
 	virtual void ActivateDefaultTool() override;
+
 	// This is currently not part of the base class at all... Should it be?
 	virtual bool IsDefaultToolActive();
 
 	// We don't actually override MouseEnter, etc, because things get forwarded to the input
 	// router via FEditorModeTools, and we don't have any additional input handling to do at the mode level.
 
-
 	// Holds the background visualiztion
 	UPROPERTY()
 	TObjectPtr<UUVEditorBackgroundPreview> BackgroundVisualization;
 
-
 protected:
 
 	// UEdMode overrides
 	virtual void CreateToolkit() override;
 	// Not sure whether we need these yet
 	virtual void BindCommands() override;
-	virtual void OnToolStarted(UInteractiveToolManager* Manager, UInteractiveTool* Tool) {}
-	virtual void OnToolEnded(UInteractiveToolManager* Manager, UInteractiveTool* Tool) {}
+	virtual void OnToolStarted(UInteractiveToolManager* Manager, UInteractiveTool* Tool) override {}
+	virtual void OnToolEnded(UInteractiveToolManager* Manager, UInteractiveTool* Tool) override {}
+	
 	void UpdateTriangleMaterialBasedOnBackground(bool IsBackgroundVisible);
-
-	// Wireframe Display Properties
-	float TriangleOpacity = 1.0;
-	float TriangleDepthOffset = 0.5;
-	const float WireframeDepthOffset = 0.6;	
-	FColor TriangleColor = FColor(50, 194, 219);
-	FColor WireframeColor = FColor(50, 100, 219);
-	FColor IslandBorderColor = FColor(103, 52, 235);
-
-    /**
-     * Stores a pointer to the specific EditorModeToolkit so we can interact with specific UI detail views
-     */
-	TSharedPtr<FUVEditorModeToolkit> EditorModeToolkit;
+	void AddDisplayedPropertySet(const TObjectPtr<UInteractiveToolPropertySet>& PropertySet);
 
 	/**
-	 * This stores the original objects that the UV editor was asked to operate on (for
-	 * instance UStaticMesh pointers).
+	 * Stores original input objects, for instance UStaticMesh pointers. AssetIDs on tool input 
+	 * objects are indices into this array (and ones that are 1:1 with it)
 	 */
 	UPROPERTY()
 	TArray<TObjectPtr<UObject>> OriginalObjectsToEdit;
 
 	/**
-	 * This stores tool targets created from OriginalObjectsToEdit that provide us with dynamic meshes.
+	 * Tool targets created from OriginalObjectsToEdit (and 1:1 with that array) that provide
+	 * us with dynamic meshes whose UV layers we unwrap.
 	 */
 	UPROPERTY()
 	TArray<TObjectPtr<UToolTarget>> ToolTargets;
 
 	/**
-	 * The actual input objects we give to the tools, including unwrapped meshes and previews, created
-	 * out of ToolTargets.
+	 * Transforms that should be used for the 3d previews, 1:1 with OriginalObjectsToEdit
+	 * and ToolTargets.
+	 */
+	TArray<FTransform> Transforms;
+
+	/**
+	 * 1:1 with the asset arrays. Hold dynamic mesh representations of the targets. These
+	 * are authoritative versions of the combined UV layers that get baked back on apply.
+	 */
+	TArray<TSharedPtr<UE::Geometry::FDynamicMesh3>> AppliedCanonicalMeshes;
+
+	/**
+	 * 1:1 with AppliedCanonicalMeshes, the actual displayed 3d meshes that can be used
+	 * by tools for background computations. However if doing so, keep in mind that while
+	 * we currently do not display more than one layer at a time for each asset, if we
+	 * someday do, tools would have to take care to disallow cases where the two layers
+	 * of the same asset might try to use the same preview for a background compute.
+	 */
+	TArray<TObjectPtr<UMeshOpPreviewWithBackgroundCompute>> AppliedPreviews;
+
+	/**
+	 * Input objects we give to the tools, one per displayed UV layer. This includes pointers
+	 * to the applied meshes, but also contains the unwrapped mesh and preview. These should
+	 * not be assumed to be the same length as the asset arrays in case we someday do not
+	 * display exactly a single layer per asset.
 	 */
 	UPROPERTY()
 	TArray<TObjectPtr<UUVEditorToolMeshInput>> ToolInputObjects;
 
-
 	/**
 	 * Wireframes have to get ticked to be able to respond to setting changes.
+	 * This is 1:1 with ToolInputObjects.
 	 */
 	TArray<TWeakObjectPtr<UMeshElementsVisualizer>> WireframesToTick;
 
+
+	// Authoritative list of targets that have changes that have not been baked back yet.
+	TSet<int32> ModifiedAssetIDs;
+
+	// Helper array of canonical unwrap changestamps that allows us to detect changes to UVs.
+	// Used to populate ModifiedAssetIDs on tick.
+	TArray<int32> LastSeenChangeStamps;
+
+
+	/** Used as a selector of UV channels/layers of opened assets in the editor. */
+	UPROPERTY()
+	TObjectPtr<UUVEditorUVChannelProperties> UVChannelProperties = nullptr;
+
+	// Used with UVChannelProperties
+	void SwitchActiveAsset(const FString& UVAsset);
+	void SwitchActiveChannel(const FString& UVChannel);
+
+	// Used to change layers with our current picker approach (we need to remember the previous
+	// layer value so we know which one to remove). Likely to change if we start adding/removing
+	// using some different UI.
+	int32 PreviousUVLayerIndex = -1;
+	int32 PendingUVLayerIndex = PreviousUVLayerIndex;
+
+
+	// Wireframe Display Properties
+	float TriangleOpacity = 1.0;
+	float TriangleDepthOffset = 0.5;
+	const float WireframeDepthOffset = 0.6;
+	FColor TriangleColor = FColor(50, 194, 219);
+	FColor WireframeColor = FColor(50, 100, 219);
+	FColor IslandBorderColor = FColor(103, 52, 235);
+
+	// Here largely for convenience to avoid having to pass it around functions.
+	UPROPERTY()
+	TObjectPtr<UWorld> LivePreviewWorld = nullptr;
+
 	/**
-	 * If we want to have mode-level property objects, whether user-visible or not,
-	 * they can be put here to be ticked.
+	* Mode-level property objects to display in the details panel.
+	*/
+	UPROPERTY()
+	TArray<TObjectPtr<UInteractiveToolPropertySet>> PropertyObjectsToDisplay;
+
+	/**
+	 * Mode-level property objects (visible or not) that get ticked.
 	 */
 	UPROPERTY()
 	TArray<TObjectPtr<UInteractiveToolPropertySet>> PropertyObjectsToTick;
 
-	// This keeps track of whether the meshes inside DisplayedMeshes have been changed since
-	// the last time that they were baked back to their targets.
-	TArray<int32> MeshChangeStamps;
+	bool bIsActive = false;
 };
+
diff --git a/Engine/Plugins/Experimental/UVEditor/Source/UVEditor/Public/UVEditorModeChannelProperties.h b/Engine/Plugins/Experimental/UVEditor/Source/UVEditor/Public/UVEditorModeChannelProperties.h
new file mode 100644
--- /dev/null
+++ b/Engine/Plugins/Experimental/UVEditor/Source/UVEditor/Public/UVEditorModeChannelProperties.h
@@ -0,0 +1,77 @@
+// Copyright Epic Games, Inc. All Rights Reserved.
+
+#pragma once
+
+#include "CoreMinimal.h"
+
+#include "InteractiveTool.h"
+#include "GeometryBase.h"
+
+#include "UVEditorModeChannelProperties.generated.h"
+
+PREDECLARE_GEOMETRY(class FDynamicMesh3);
+
+class UInteractiveToolPropertySet;
+class UToolTarget;
+class UUVEditorToolMeshInput;
+
+/**
+ * UV Layer Settings for the UVEditorMode
+ */
+UCLASS()
+class UVEDITOR_API UUVEditorUVChannelProperties : public UInteractiveToolPropertySet
+{
+	GENERATED_BODY()
+public:
+
+	UPROPERTY(EditAnywhere, Category = UVChannel, meta = (DisplayName = "Asset", GetOptions = GetAssetNames))
+	FString Asset;
+
+	UFUNCTION()
+	const TArray<FString>& GetAssetNames();
+
+	UPROPERTY(EditAnywhere, Category = UVChannel, meta = (DisplayName = "UV Channel", GetOptions = GetUVChannelNames))
+	FString UVChannel;
+
+	UFUNCTION()
+	const TArray<FString>& GetUVChannelNames();
+
+	TArray<FString> UVChannelNames;
+	TArray<FString> UVAssetNames;
+
+	// 1:1 with UVAssetNames
+	TArray<int32> NumUVLayersPerAsset;
+
+
+public:
+	void Initialize(
+		TArray<TObjectPtr<UToolTarget>>& TargetObjects,
+		TArray<TSharedPtr<UE::Geometry::FDynamicMesh3>> AppliedCanonicalMeshes,
+		bool bInitializeSelection);
+
+	/**
+	 * Verify that the UV asset selection is valid
+	 * @param bUpdateIfInvalid if selection is not valid, reset UVAsset to Asset0 or empty if no assets exist
+	 * @return true if selection in UVAsset is an entry in UVAssetNames.
+	 */
+	bool ValidateUVAssetSelection(bool bUpdateIfInvalid);
+
+	/**
+	 * Verify that the UV channel selection is valid
+	 * @param bUpdateIfInvalid if selection is not valid, UVChannel to UV0 or empty if no UV channels exist
+	 * @return true if selection in UVChannel is an entry in UVChannelNamesList.
+	 */
+	bool ValidateUVChannelSelection(bool bUpdateIfInvalid);
+
+
+	/**
+	 * @return selected UV asset ID, or -1 if invalid selection
+	 */
+	int32 GetSelectedAssetID();
+
+	/**
+	 * @param bForceToZeroOnFailure if true, then instead of returning -1 we return 0 so calling code can fallback to default UV paths
+	 * @return selected UV channel index, or -1 if invalid selection, or 0 if invalid selection and bool bForceToZeroOnFailure = true
+	 */
+	int32 GetSelectedChannelIndex(bool bForceToZeroOnFailure = false);
+};
diff --git a/Engine/Plugins/Experimental/UVEditor/Source/UVEditor/Public/UVEditorModeToolkit.h b/Engine/Plugins/Experimental/UVEditor/Source/UVEditor/Public/UVEditorModeToolkit.h
--- a/Engine/Plugins/Experimental/UVEditor/Source/UVEditor/Public/UVEditorModeToolkit.h
+++ b/Engine/Plugins/Experimental/UVEditor/Source/UVEditor/Public/UVEditorModeToolkit.h
@@ -9,6 +9,7 @@
 
 class SBorder;
 class STextBlock;
+class UInteractiveToolPropertySet;
 
 /**
  * The UV editor mode toolkit is responsible for the panel on the side in the UV editor
@@ -28,8 +29,7 @@ public:
 	virtual FText GetBaseToolkitName() const override;
 	virtual TSharedPtr<class SWidget> GetInlineContent() const override { return ToolkitWidget; }
 
-
-	virtual void SetBackgroundSettings(UObject* InSettingsObject);
+	virtual void SetModeDetailsViewObjects(const TArray<TObjectPtr<UInteractiveToolPropertySet>>& InObjects);
 
 	virtual FText GetActiveToolDisplayName() const override { return ActiveToolName; }
 	virtual void OnToolStarted(UInteractiveToolManager* Manager, UInteractiveTool* Tool) override;
diff --git a/Engine/Plugins/Experimental/UVEditor/Source/UVEditorTools/Private/ToolTargets/UVEditorToolMeshInput.cpp b/Engine/Plugins/Experimental/UVEditor/Source/UVEditorTools/Private/ToolTargets/UVEditorToolMeshInput.cpp
--- a/Engine/Plugins/Experimental/UVEditor/Source/UVEditorTools/Private/ToolTargets/UVEditorToolMeshInput.cpp
+++ b/Engine/Plugins/Experimental/UVEditor/Source/UVEditorTools/Private/ToolTargets/UVEditorToolMeshInput.cpp
@@ -168,13 +168,13 @@ bool UUVEditorToolMeshInput::IsValid() const
 		&& UnwrapPreview && UnwrapPreview->IsValidLowLevel()
 		&& AppliedCanonical
 		&& AppliedPreview && AppliedPreview->IsValidLowLevel()
-		&& OriginalAsset && OriginalAsset->IsValidLowLevel()
+		&& SourceTarget && SourceTarget->IsValid()
 		&& UVLayerIndex >= 0;
 }
 
-bool UUVEditorToolMeshInput::InitializeMeshes(
-	UToolTarget* Target, int32 UVLayerIndexIn,
-	UWorld* UnwrapWorld, UWorld* LivePreviewWorld,
+bool UUVEditorToolMeshInput::InitializeMeshes(UToolTarget* Target, 
+	TSharedPtr<FDynamicMesh3> AppliedCanonicalIn, UMeshOpPreviewWithBackgroundCompute* AppliedPreviewIn,
+	int32 AssetIDIn, int32 UVLayerIndexIn, UWorld* UnwrapWorld, UWorld* LivePreviewWorld,
 	UMaterialInterface* WorkingMaterialIn,
 	TFunction<FVector3d(const FVector2f&)> UVToVertPositionFuncIn,
 	TFunction<FVector2f(const FVector3d&)> VertPositionToUVFuncIn)
@@ -184,17 +184,14 @@ bool UUVEditorToolMeshInput::InitializeMeshes(
 	// We'll need to update this function once they exist.
 	using namespace UE::ToolTarget;
 
-	checkSlow(Cast<IAssetBackedTarget>(Target) 
-		&& Cast<IDynamicMeshProvider>(Target)
-		&& Cast<IMaterialProvider>(Target));
-
-	OriginalAsset = Cast<IAssetBackedTarget>(Target)->GetSourceData();
+	SourceTarget = Target;
+	AssetID = AssetIDIn;
 	UVLayerIndex = UVLayerIndexIn;
 	UVToVertPosition = UVToVertPositionFuncIn;
 	VertPositionToUV = VertPositionToUVFuncIn;
 
 	// We are given the preview- i.e. the mesh with the uv layer applied.
-	AppliedCanonical = Cast<IDynamicMeshProvider>(Target)->GetDynamicMesh();
+	AppliedCanonical = AppliedCanonicalIn;
 
 	if (!AppliedCanonical->HasAttributes()
 		|| UVLayerIndex >= AppliedCanonical->Attributes()->NumUVLayers())
@@ -202,14 +199,7 @@ bool UUVEditorToolMeshInput::InitializeMeshes(
 		return false;
 	}
 
-	// Create a preview version of the 3d mesh
-	AppliedPreview = NewObject<UMeshOpPreviewWithBackgroundCompute>();
-	AppliedPreview->Setup(LivePreviewWorld);
-	AppliedPreview->PreviewMesh->UpdatePreview(AppliedCanonical.Get());
-
-	FComponentMaterialSet MaterialSet;
-	Cast<IMaterialProvider>(Target)->GetMaterialSet(MaterialSet);
-	AppliedPreview->ConfigureMaterials(MaterialSet.Materials, WorkingMaterialIn);
+	AppliedPreview = AppliedPreviewIn;
 
 	// Set up the unwrapped mesh
 	UnwrapCanonical = MakeShared<FDynamicMesh3>();
@@ -235,15 +225,14 @@ void UUVEditorToolMeshInput::Shutdown()
 		WireframeDisplay = nullptr;
 	}
 
-	UnwrapPreview->Cancel();
-	AppliedPreview->Cancel();
-
 	UnwrapCanonical = nullptr;
+	UnwrapPreview->Shutdown();
 	UnwrapPreview = nullptr;
 	AppliedCanonical = nullptr;
+	// Can't shut down AppliedPreview because it is owned by mode
 	AppliedPreview = nullptr;
 
-	OriginalAsset = nullptr;
+	SourceTarget = nullptr;
 }
 
 void UUVEditorToolMeshInput::UpdateUnwrapPreviewOverlayFromPositions(const TArray<int32>* ChangedVids, const TArray<int32>* ChangedConnectivityTids, const TArray<int32>* FastRenderUpdateTids)
@@ -289,7 +278,7 @@ void UUVEditorToolMeshInput::UpdateAppliedPreviewFromUnwrapPreview(const TArray<
 	AppliedPreview->PreviewMesh->DeferredEditMesh([this, SourceUnwrapMesh, ChangedVids, ChangedConnectivityTids](FDynamicMesh3& Mesh)
 	{
 			const FDynamicMeshUVOverlay* SourceOverlay = SourceUnwrapMesh->Attributes()->PrimaryUV();
-			FDynamicMeshUVOverlay* DestOverlay = Mesh.Attributes()->PrimaryUV();
+			FDynamicMeshUVOverlay* DestOverlay = Mesh.Attributes()->GetUVLayer(UVLayerIndex);
 
 			CopyMeshOverlay(*SourceOverlay, *DestOverlay, false, ChangedVids, ChangedConnectivityTids);
 	}, false);
@@ -310,7 +299,7 @@ void UUVEditorToolMeshInput::UpdateUnwrapPreviewFromAppliedPreview(const TArray<
 	using namespace UVEditorToolMeshInputLocals;
 
 	// Convert the AppliedPreview UV overlay to positions in UnwrapPreview
-	const FDynamicMeshUVOverlay* SourceOverlay = AppliedPreview->PreviewMesh->GetMesh()->Attributes()->PrimaryUV();
+	const FDynamicMeshUVOverlay* SourceOverlay = AppliedPreview->PreviewMesh->GetMesh()->Attributes()->GetUVLayer(UVLayerIndex);
 	UnwrapPreview->PreviewMesh->DeferredEditMesh([this, SourceOverlay, ChangedElementIDs, ChangedConnectivityTids](FDynamicMesh3& Mesh)
 	{
 		UVEditorToolUtil::UpdateUVUnwrapMesh(*SourceOverlay, Mesh, UVToVertPosition, ChangedElementIDs, ChangedConnectivityTids);
@@ -410,6 +399,18 @@ void UUVEditorToolMeshInput::UpdateAllFromUnwrapCanonical(
 	UpdatePreviewsFromCanonical(ChangedVids, ChangedConnectivityTids, FastRenderUpdateTids);
 }
 
+void UUVEditorToolMeshInput::UpdateAllFromAppliedCanonical(
+	const TArray<int32>* ChangedElementIDs, const TArray<int32>* ChangedConnectivityTids, const TArray<int32>* FastRenderUpdateTids)
+{
+	using namespace UVEditorToolMeshInputLocals;
+
+	// Update UnwrapCanonical
+	const FDynamicMeshUVOverlay* SourceOverlay = AppliedCanonical->Attributes()->GetUVLayer(UVLayerIndex);
+	UVEditorToolUtil::UpdateUVUnwrapMesh(*SourceOverlay, *UnwrapCanonical, UVToVertPosition, ChangedElementIDs, ChangedConnectivityTids);
+
+	UpdatePreviewsFromCanonical(ChangedElementIDs, ChangedConnectivityTids, FastRenderUpdateTids);
+}
+
 void UUVEditorToolMeshInput::UpdateAllFromAppliedPreview(
 	const TArray<int32>* ChangedElementIDs, const TArray<int32>* ChangedConnectivityTids, const TArray<int32>* FastRenderUpdateTids)
 {
diff --git a/Engine/Plugins/Experimental/UVEditor/Source/UVEditorTools/Private/UVSelectTool.cpp b/Engine/Plugins/Experimental/UVEditor/Source/UVEditorTools/Private/UVSelectTool.cpp
--- a/Engine/Plugins/Experimental/UVEditor/Source/UVEditorTools/Private/UVSelectTool.cpp
+++ b/Engine/Plugins/Experimental/UVEditor/Source/UVEditorTools/Private/UVSelectTool.cpp
@@ -81,6 +81,12 @@ namespace UVSelectToolLocals
 			}
 		}
 
+		virtual bool HasExpired(UObject* Object) const override
+		{
+			UUVSelectToolChangeRouter* ChangeRouter = Cast<UUVSelectToolChangeRouter>(Object);
+			return !(ChangeRouter && ChangeRouter->CurrentSelectTool.IsValid());
+		}
+
 		virtual FString ToString() const override
 		{
 			return TEXT("UVSelectToolLocals::FSelectionChange");
@@ -158,7 +164,7 @@ namespace UVSelectToolLocals
 
 		virtual bool HasExpired(UObject* Object) const override
 		{
-			return !(UVToolInputObject.IsValid() && UnwrapCanonicalMeshChange);
+			return !(UVToolInputObject.IsValid() && UVToolInputObject->IsValid() && UnwrapCanonicalMeshChange);
 		}
 
 
@@ -285,7 +291,14 @@ void UUVSelectTool::Setup()
 			Tree->SetMesh(Target->UnwrapCanonical.Get());
 			TreeStore->Set(Target->UnwrapCanonical.Get(), Tree);
 		}
-		if (!Tree->IsValid(false))
+		// TODO: We currently can't do this check because the aabb tree checks the mesh
+		// changestamp, and that gets (arguably incorrectly) reset if we do an update
+		// that involves a clear. So, for instance, after a layout operation, the changestamp
+		// frequently ends up matching the original and incorrectly passes the check.
+		// All this means that storing the tree is currently pointless because
+		// we rebuild it each time we switch back to this tool, but we keep the code
+		// until we fix the change stamp thing or give a workaround.
+		//if (!Tree->IsValid(false))
 		{
 			Tree->Build();
 		}
@@ -299,6 +312,19 @@ void UUVSelectTool::Setup()
 			Targets[i]->UnwrapPreview->PreviewMesh->GetTransform());
 	}
 
+	// Remove trees that are no longer relevant (because their layer is not displayed)
+	TreeStore->RemoveByPredicate([this](TPair<FDynamicMesh3*, TSharedPtr<FDynamicMeshAABBTree3>> Pair)
+	{
+		for (TObjectPtr<UUVEditorToolMeshInput> Target : Targets)
+		{
+			if (Target->UnwrapCanonical.Get() == Pair.Key)
+			{
+				return false;
+			}
+		}
+		return true;
+	});
+
 	// Make sure that if we receive undo/redo events on the meshes, we update the tree structures
 	// and the selection mechanic drawn elements. Note that we mainly have to worry about this
 	// because the select tool is the default UV editor tool, and as such it can receive undo
diff --git a/Engine/Plugins/Experimental/UVEditor/Source/UVEditorTools/Private/UVToolContextObjects.cpp b/Engine/Plugins/Experimental/UVEditor/Source/UVEditorTools/Private/UVToolContextObjects.cpp
--- a/Engine/Plugins/Experimental/UVEditor/Source/UVEditorTools/Private/UVToolContextObjects.cpp
+++ b/Engine/Plugins/Experimental/UVEditor/Source/UVEditorTools/Private/UVToolContextObjects.cpp
@@ -44,7 +44,7 @@ namespace UVToolContextObjectLocals
 
 		virtual bool HasExpired(UObject* Object) const override
 		{
-			return !(UVToolInputObject.IsValid() && UnwrapCanonicalMeshChange);
+			return !(UVToolInputObject.IsValid() && UVToolInputObject->IsValid() && UnwrapCanonicalMeshChange);
 		}
 
 
diff --git a/Engine/Plugins/Experimental/UVEditor/Source/UVEditorTools/Public/ToolTargets/UVEditorToolMeshInput.h b/Engine/Plugins/Experimental/UVEditor/Source/UVEditorTools/Public/ToolTargets/UVEditorToolMeshInput.h
--- a/Engine/Plugins/Experimental/UVEditor/Source/UVEditorTools/Public/ToolTargets/UVEditorToolMeshInput.h
+++ b/Engine/Plugins/Experimental/UVEditor/Source/UVEditorTools/Public/ToolTargets/UVEditorToolMeshInput.h
@@ -89,14 +89,18 @@ public:
 	FOnUndoRedo OnUndoRedo;
 
 	// Additional needed information
-	TObjectPtr<UObject> OriginalAsset = nullptr;
+	TObjectPtr<UToolTarget> SourceTarget = nullptr;
+	int32 AssetID = -1;
 	int32 UVLayerIndex = 0;
 
 	// Mappings used for generating and baking back the unwrap.
 	TFunction<FVector3d(const FVector2f&)> UVToVertPosition;
 	TFunction<FVector2f(const FVector3d&)> VertPositionToUV;
 
-	bool InitializeMeshes(UToolTarget* Target, int32 UVLayerIndex,
+	bool InitializeMeshes(UToolTarget* Target, 
+		TSharedPtr<UE::Geometry::FDynamicMesh3> AppliedCanonicalIn,
+		UMeshOpPreviewWithBackgroundCompute* AppliedPreviewIn,
+		int32 AssetID, int32 UVLayerIndex,
 		UWorld* UnwrapWorld, UWorld* LivePreviewWorld,
 		UMaterialInterface* WorkingMaterialIn,
 		TFunction<FVector3d(const FVector2f&)> UVToVertPositionFuncIn,
@@ -167,6 +171,12 @@ public:
 	 */
 	void UpdateAllFromUnwrapCanonical(const TArray<int32>* ChangedVids = nullptr, const TArray<int32>* ChangedConnectivityTids = nullptr, const TArray<int32>* FastRenderUpdateTids = nullptr);
 
+	/**
+	 * Updates the other meshes using the mesh in AppliedCanonical. Useful when switching UV layer indices or otherwise
+	 * resetting the collection from the "original" data.
+	 */
+	void UpdateAllFromAppliedCanonical(const TArray<int32>* ChangedVids = nullptr, const TArray<int32>* ChangedConnectivityTids = nullptr, const TArray<int32>* FastRenderUpdateTids = nullptr);
+
 	/**
 	 * Updates the other meshes using the UV overlay in the live preview.
 	 */
