UVEditor: Add support for UDynamicMesh, and make it possible to open skeletal meshes from content browser.

#rnx

diff --git a/Engine/Plugins/Experimental/UVEditor/Source/UVEditor/Private/UVEditorModule.cpp b/Engine/Plugins/Experimental/UVEditor/Source/UVEditor/Private/UVEditorModule.cpp
--- a/Engine/Plugins/Experimental/UVEditor/Source/UVEditor/Private/UVEditorModule.cpp
+++ b/Engine/Plugins/Experimental/UVEditor/Source/UVEditor/Private/UVEditorModule.cpp
@@ -4,6 +4,7 @@
 
 #include "ContentBrowserMenuContexts.h"
 #include "LevelEditorMenuContext.h"
+#include "Selection.h"
 #include "UVEditor.h"
 #include "UVEditorCommands.h"
 #include "UVEditorMode.h"
@@ -39,13 +40,11 @@ void FUVEditorModule::RegisterMenus()
 	// Allows cleanup when module unloads.
 	FToolMenuOwnerScoped OwnerScoped(this);
 
-	// Extent the content browser context menu for static meshes
+	// Extend the content browser context menu for static meshes and skeletal meshes
+	auto AddToContextMenuSection = [this](FToolMenuSection& Section)
 	{
-		UToolMenu* Menu = UToolMenus::Get()->ExtendMenu("ContentBrowser.AssetContextMenu.StaticMesh");
-		FToolMenuSection& Section = Menu->FindOrAddSection("GetAssetActions");
-
 		Section.AddDynamicEntry("OpenUVEditor", FNewToolMenuSectionDelegate::CreateLambda(
-			[this](FToolMenuSection& Section) 
+			[](FToolMenuSection& Section)
 			{
 				// We'll need to get the target objects out of the context
 				if (UContentBrowserAssetContextMenuContext* Context = Section.FindContext<UContentBrowserAssetContextMenuContext>())
@@ -67,30 +66,64 @@ void FUVEditorModule::RegisterMenus()
 					}
 				}
 			}));
+	};
+	{
+		UToolMenu* Menu = UToolMenus::Get()->ExtendMenu("ContentBrowser.AssetContextMenu.StaticMesh");
+		FToolMenuSection& Section = Menu->FindOrAddSection("GetAssetActions");
+		AddToContextMenuSection(Section);
+	}
+	{
+		UToolMenu* Menu = UToolMenus::Get()->ExtendMenu("ContentBrowser.AssetContextMenu.SkeletalMesh");
+		FToolMenuSection& Section = Menu->FindOrAddSection("GetAssetActions");
+		AddToContextMenuSection(Section);
 	}
 
-	// Extent the level editor context menu
+	// Extend the level editor context menu
+	// TODO: Currently, for dynamic meshes, the "Open UV Editor" dialog ends up in an unrelated section
+	// because there is no ActorAsset section. Need to figure out a way to change which section the
+	// option goes into based on whether the ActorAsset section exists...
 	{
 		UToolMenu* Menu = UToolMenus::Get()->ExtendMenu("LevelEditor.ActorContextMenu");
 		FToolMenuSection& Section = Menu->FindOrAddSection("ActorAsset");
 		Section.AddDynamicEntry("OpenUVEditor", FNewToolMenuSectionDelegate::CreateLambda(
 			[this](FToolMenuSection& Section) 
 			{
-				TArray<UObject*> ReferencedAssets;
-				GEditor->GetReferencedAssetsForEditorSelection(ReferencedAssets);
+				TArray<TObjectPtr<UObject>> TargetObjects;
+
+				// TODO: There's some newer way to iterate across selected actors that we can switch over to someday
+				for (FSelectionIterator It(GEditor->GetSelectedActorIterator()); It; ++It)
+				{
+					// We are interested in the (unique) assets backing the actor, or else the actor
+					// itself if it is not asset backed (such as UDynamicMesh).
+					AActor* Actor = static_cast<AActor*>(*It);
+					TArray<UObject*> ActorAssets;
+					Actor->GetReferencedContentObjects(ActorAssets);
+
+					if (ActorAssets.Num() > 0)
+					{
+						for (UObject* Asset : ActorAssets)
+						{
+							TargetObjects.AddUnique(Asset);
+						}
+					}
+					else
+					{
+						// Need to transform actors to components here because this is what our tool targets expect
+						TInlineComponentArray<UActorComponent*> ActorComponents;
+						Actor->GetComponents(ActorComponents);
+						TargetObjects.Append(ActorComponents);
+					}
+				}
 
 				UUVEditorSubsystem* UVSubsystem = GEditor->GetEditorSubsystem<UUVEditorSubsystem>();
 				check(UVSubsystem);
 
-				if (UVSubsystem->AreObjectsValidTargets(ReferencedAssets))
+				if (UVSubsystem->AreObjectsValidTargets(TargetObjects))
 				{
-					TArray<TObjectPtr<UObject>> AssetsToEdit;
-					AssetsToEdit.Append(ReferencedAssets);
-
 					TSharedPtr<class FUICommandList> CommandListToBind = MakeShared<FUICommandList>();
 					CommandListToBind->MapAction(
 						FUVEditorCommands::Get().OpenUVEditor,
-						FExecuteAction::CreateUObject(UVSubsystem, &UUVEditorSubsystem::StartUVEditor, AssetsToEdit));
+						FExecuteAction::CreateUObject(UVSubsystem, &UUVEditorSubsystem::StartUVEditor, TargetObjects));
 
 					Section.AddMenuEntryWithCommandList(FUVEditorCommands::Get().OpenUVEditor, CommandListToBind);
 				}
diff --git a/Engine/Plugins/Experimental/UVEditor/Source/UVEditor/Private/UVEditorSubsystem.cpp b/Engine/Plugins/Experimental/UVEditor/Source/UVEditor/Private/UVEditorSubsystem.cpp
--- a/Engine/Plugins/Experimental/UVEditor/Source/UVEditor/Private/UVEditorSubsystem.cpp
+++ b/Engine/Plugins/Experimental/UVEditor/Source/UVEditor/Private/UVEditorSubsystem.cpp
@@ -3,6 +3,7 @@
 #include "UVEditorSubsystem.h"
 
 #include "ToolTargetManager.h"
+#include "ToolTargets/DynamicMeshComponentToolTarget.h"
 #include "ToolTargets/StaticMeshToolTarget.h"
 #include "ToolTargets/SkeletalMeshToolTarget.h"
 #include "UVEditor.h"
@@ -19,6 +20,7 @@ void UUVEditorSubsystem::Initialize(FSubsystemCollectionBase& Collection)
 
 	ToolTargetManager->AddTargetFactory(NewObject<UStaticMeshToolTargetFactory>(ToolTargetManager));
 	ToolTargetManager->AddTargetFactory(NewObject<USkeletalMeshToolTargetFactory>(ToolTargetManager));
+	ToolTargetManager->AddTargetFactory(NewObject<UDynamicMeshComponentToolTargetFactory>(ToolTargetManager));
 }
 
 void UUVEditorSubsystem::Deinitialize()
