No unity fixes from partner
#rnx

diff --git a/Engine/Plugins/Experimental/UVEditor/Source/UVEditor/Private/SUVEditor2DViewportToolBar.cpp b/Engine/Plugins/Experimental/UVEditor/Source/UVEditor/Private/SUVEditor2DViewportToolBar.cpp
--- a/Engine/Plugins/Experimental/UVEditor/Source/UVEditor/Private/SUVEditor2DViewportToolBar.cpp
+++ b/Engine/Plugins/Experimental/UVEditor/Source/UVEditor/Private/SUVEditor2DViewportToolBar.cpp
@@ -5,6 +5,9 @@
 #include "SUVEditor2DViewportToolBar.h"
 
 #include "EditorViewportCommands.h"
+#include "Framework/MultiBox/MultiBoxBuilder.h"
+#include "Widgets/Layout/SBorder.h"
+#include "Widgets/SBoxPanel.h"
 
 
 void SUVEditor2DViewportToolBar::Construct(const FArguments& InArgs)
diff --git a/Engine/Plugins/Experimental/UVEditor/Source/UVEditor/Public/SUVEditor2DViewportToolBar.h b/Engine/Plugins/Experimental/UVEditor/Source/UVEditor/Public/SUVEditor2DViewportToolBar.h
--- a/Engine/Plugins/Experimental/UVEditor/Source/UVEditor/Public/SUVEditor2DViewportToolBar.h
+++ b/Engine/Plugins/Experimental/UVEditor/Source/UVEditor/Public/SUVEditor2DViewportToolBar.h
@@ -6,6 +6,7 @@
 
 #include "SViewportToolBar.h"
 
+class FExtender;
 class FUICommandList;
 class SUVEditor2DViewport;
 
diff --git a/Engine/Source/Runtime/Interchange/Engine/Private/InterchangeManager.cpp b/Engine/Source/Runtime/Interchange/Engine/Private/InterchangeManager.cpp
--- a/Engine/Source/Runtime/Interchange/Engine/Private/InterchangeManager.cpp
+++ b/Engine/Source/Runtime/Interchange/Engine/Private/InterchangeManager.cpp
@@ -4,6 +4,7 @@
 
 #include "AssetDataTagMap.h"
 #include "AssetRegistryModule.h"
+#include "Async/Async.h"
 #include "CoreMinimal.h"
 #include "Engine/Blueprint.h"
 #include "Framework/Notifications/NotificationManager.h"
