UVEditor: Channel Editing Tool

#rnx

diff --git a/Engine/Plugins/Experimental/UVEditor/Source/UVEditor/Private/UVEditorCommands.cpp b/Engine/Plugins/Experimental/UVEditor/Source/UVEditor/Private/UVEditorCommands.cpp
--- a/Engine/Plugins/Experimental/UVEditor/Source/UVEditor/Private/UVEditorCommands.cpp
+++ b/Engine/Plugins/Experimental/UVEditor/Source/UVEditor/Private/UVEditorCommands.cpp
@@ -26,6 +26,7 @@ void FUVEditorCommands::RegisterCommands()
 	UI_COMMAND(BeginSelectTool, "Select", "Switch to select tool", EUserInterfaceActionType::ToggleButton, FInputChord());
 	UI_COMMAND(BeginLayoutTool, "Layout", "Switch to layout tool", EUserInterfaceActionType::ToggleButton, FInputChord());
 	UI_COMMAND(BeginParameterizeMeshTool, "Auto UV", "Switch to Auto UV tool", EUserInterfaceActionType::ToggleButton, FInputChord());
+	UI_COMMAND(BeginChannelEditTool, "Channel Edit", "Switch to Channel Edit tool", EUserInterfaceActionType::ToggleButton, FInputChord());
 
 	// These allow us to link up to pressed keys
 	UI_COMMAND(AcceptActiveTool, "Accept", "Accept the active tool", EUserInterfaceActionType::Button, FInputChord(EKeys::Enter));
diff --git a/Engine/Plugins/Experimental/UVEditor/Source/UVEditor/Private/UVEditorMode.cpp b/Engine/Plugins/Experimental/UVEditor/Source/UVEditor/Private/UVEditorMode.cpp
--- a/Engine/Plugins/Experimental/UVEditor/Source/UVEditor/Private/UVEditorMode.cpp
+++ b/Engine/Plugins/Experimental/UVEditor/Source/UVEditor/Private/UVEditorMode.cpp
@@ -23,6 +23,7 @@
 #include "UVEditorCommands.h"
 #include "UVEditorLayoutTool.h"
 #include "UVEditorParameterizeMeshTool.h"
+#include "UVEditorLayerEditTool.h"
 #include "UVSelectTool.h"
 #include "UVEditorModeToolkit.h"
 #include "UVEditorSubsystem.h"
@@ -137,8 +138,8 @@ void UUVEditorMode::Enter()
 	AddDisplayedPropertySet(UVChannelProperties);
 	AddDisplayedPropertySet(BackgroundVisualization->Settings);
 		
-	StaticCastSharedPtr<FUVEditorModeToolkit>(Toolkit)->SetModeDetailsViewObjects(PropertyObjectsToDisplay);
-
+	StaticCastSharedPtr<FUVEditorModeToolkit>(Toolkit)->SetModeDetailsViewObjects(PropertyObjectsToDisplay);	
+	
 	RegisterTools();
 
 	ActivateDefaultTool();
@@ -161,6 +162,10 @@ void UUVEditorMode::RegisterTools()
 	UUVEditorParameterizeMeshToolBuilder* UVEditorParameterizeMeshToolBuilder = NewObject<UUVEditorParameterizeMeshToolBuilder>();
 	UVEditorParameterizeMeshToolBuilder->Targets = &ToolInputObjects;
 	RegisterTool(CommandInfos.BeginParameterizeMeshTool, TEXT("UVParameterizeMeshTool"), UVEditorParameterizeMeshToolBuilder);
+
+	UUVEditorChannelEditToolBuilder* UVEditorChannelEditToolBuilder = NewObject<UUVEditorChannelEditToolBuilder>();
+	UVEditorChannelEditToolBuilder->Targets = &ToolInputObjects;
+	RegisterTool(CommandInfos.BeginChannelEditTool, TEXT("UVChannelEditTool"), UVEditorChannelEditToolBuilder);
 }
 
 void UUVEditorMode::AddDisplayedPropertySet(const TObjectPtr<UInteractiveToolPropertySet>& PropertySet)
@@ -377,9 +382,32 @@ void UUVEditorMode::InitializeTargets(const TArray<TObjectPtr<UObject>>& AssetsI
 	}
 
 	// Initialize our layer selector
-	UVChannelProperties->Initialize(ToolTargets, AppliedCanonicalMeshes, true);
-	PreviousUVLayerIndex = UVChannelProperties->GetSelectedChannelIndex();
-	PendingUVLayerIndex = PreviousUVLayerIndex;
+	UVChannelProperties->Initialize(ToolTargets, AppliedCanonicalMeshes, true);	
+	PreviousUVLayerIndex.SetNumZeroed(ToolTargets.Num());
+	PendingUVLayerIndex.SetNumZeroed(ToolTargets.Num());		
+
+	int32 AssetID = UVChannelProperties->GetSelectedAssetID();
+	UUVToolAssetAndChannelAPI* AssetAndLayerAPI = ContextStore->FindContext<UUVToolAssetAndChannelAPI>();
+	if (AssetAndLayerAPI)
+	{
+		AssetAndLayerAPI->RequestChannelVisibilityChangeFunc = [this](const TArray<int32>& LayerPerAsset, bool bForceRebuildUnwrap, bool bEmitUndoTransaction) {
+			this->ForceUpdateDisplayChannel(LayerPerAsset, bForceRebuildUnwrap, bEmitUndoTransaction);
+		};
+
+		AssetAndLayerAPI->NotifyOfAssetChannelCountChangeFunc = [this](int32 AssetID) {
+			UVChannelProperties->Initialize(ToolTargets, AppliedCanonicalMeshes, false);
+		};
+
+		AssetAndLayerAPI->GetCurrentChannelVisibilityFunc = [this]() {
+			TArray<int32> VisibleLayers;
+			VisibleLayers.SetNum(ToolTargets.Num());
+			for (int32 AssetID = 0; AssetID < ToolTargets.Num(); ++AssetID)
+			{
+				VisibleLayers[AssetID] = ToolInputObjects[AssetID]->UVLayerIndex;
+			}
+			return VisibleLayers;
+		};
+	}
 }
 
 void UUVEditorMode::EmitToolIndependentObjectChange(UObject* TargetObject, TUniquePtr<FToolCommandChange> Change, const FText& Description)
@@ -451,31 +479,42 @@ void UUVEditorMode::ModeTick(float DeltaTime)
 
 	Super::ModeTick(DeltaTime);
 
-	if (PreviousUVLayerIndex!= PendingUVLayerIndex)
+	bool bStartDefaultTool = false;
+	GetToolManager()->BeginUndoTransaction(UVLayerChangeTransactionName);
+	for (int32 AssetID = 0; AssetID < ToolInputObjects.Num(); ++AssetID)
 	{
-		int32 AssetID = UVChannelProperties->GetSelectedAssetID();
-		if (ensure(ToolInputObjects[AssetID]->UVLayerIndex != PendingUVLayerIndex))
+		if ((PreviousUVLayerIndex[AssetID] != PendingUVLayerIndex[AssetID] &&
+			ToolInputObjects[AssetID]->UVLayerIndex != PendingUVLayerIndex[AssetID]) || bForceRebuildUVLayer)
 		{
-			GetToolManager()->BeginUndoTransaction(UVLayerChangeTransactionName);
+			
 
 			// TODO: Perhaps we need our own interactive tools context that allows this kind of "end tool now"
 			// call. We can't do the normal EndTool call because we cannot defer shutdown until the tick.
-			GetInteractiveToolsContext()->ToolManager->DeactivateTool(EToolSide::Mouse, EToolShutdownType::Cancel);
+			if (bPendingUVLayerChangeDestroyTool)
+			{
+				GetInteractiveToolsContext()->ToolManager->DeactivateTool(EToolSide::Mouse, EToolShutdownType::Cancel);
+			}
 
-			ChangeInputObjectLayer(AssetID, PendingUVLayerIndex);
+			ChangeInputObjectLayer(AssetID, PendingUVLayerIndex[AssetID], bForceRebuildUVLayer);
 
 			GetInteractiveToolsContext()->GetTransactionAPI()->AppendChange(this,
-				MakeUnique<FInputObjectUVLayerChange>(AssetID, PreviousUVLayerIndex, PendingUVLayerIndex), 
-				UVLayerChangeTransactionName);
-
-			GetToolManager()->EndUndoTransaction();
-
-			ActivateDefaultTool();
-			//GetInteractiveToolsContext()->ToolManager->ActivateTool(EToolSide::Mouse);
+				MakeUnique<FInputObjectUVLayerChange>(AssetID, PreviousUVLayerIndex[AssetID], PendingUVLayerIndex[AssetID]),
+				UVLayerChangeTransactionName);		
 
+			if (bPendingUVLayerChangeDestroyTool)
+			{
+				bStartDefaultTool = true;
+			}
+		
+			PreviousUVLayerIndex[AssetID] = PendingUVLayerIndex[AssetID];			
 		}
-		PreviousUVLayerIndex = PendingUVLayerIndex;
 	}
+	if (bStartDefaultTool)
+	{
+		ActivateDefaultTool();
+	}
+	GetToolManager()->EndUndoTransaction();
+	bForceRebuildUVLayer = false;
 
 	for (TObjectPtr<UInteractiveToolPropertySet>& Propset : PropertyObjectsToTick)
 	{
@@ -528,6 +567,8 @@ void UUVEditorMode::SwitchActiveAsset(const FString& UVAsset)
 		// Not doing an ensure here because the "revert to default" can give us an empty string
 		UVChannelProperties->ValidateUVAssetSelection(true);
 		UpdateSelectedLayer();
+		bPendingUVLayerChangeDestroyTool = true;
+		bForceRebuildUVLayer = false;
 	}
 }
 
@@ -550,17 +591,18 @@ void UUVEditorMode::UpdateSelectedLayer()
 
 			UVChannelProperties->UVChannel = ChannelNames[InputObject->UVLayerIndex];
 			UVChannelProperties->SilentUpdateWatched();
-			PreviousUVLayerIndex = InputObject->UVLayerIndex;
+			PreviousUVLayerIndex[AssetID] = InputObject->UVLayerIndex;
 			break;
 		}
 	}
 	if (!ensure(bFound))
 	{
 		UVChannelProperties->UVChannel = TEXT("");
-		PreviousUVLayerIndex = INDEX_NONE;
 	}
-
-	PendingUVLayerIndex = PreviousUVLayerIndex;
+	else 
+	{
+		PendingUVLayerIndex[AssetID] = PreviousUVLayerIndex[AssetID];
+	}		
 }
 
 
@@ -573,15 +615,22 @@ void UUVEditorMode::SwitchActiveChannel(const FString& UVChannel)
 		// Not doing an ensure because the "revert to default" can give us an empty string
 		UVChannelProperties->ValidateUVChannelSelection(true);
 		int32 NewUVLayerIndex = UVChannelProperties->GetSelectedChannelIndex();
+		int32 AssetID = UVChannelProperties->GetSelectedAssetID();
+		if (!ensure(AssetID != INDEX_NONE))
+		{
+			return;
+		}
 
 		if (ensure(NewUVLayerIndex != IndexConstants::InvalidID))
 		{
-			PendingUVLayerIndex = NewUVLayerIndex;
+			PendingUVLayerIndex[AssetID] = NewUVLayerIndex;
 		}
 	}
+	bPendingUVLayerChangeDestroyTool = true;
+	bForceRebuildUVLayer = false;
 }
 
-void UUVEditorMode::ChangeInputObjectLayer(int32 AssetID, int32 NewLayerIndex)
+void UUVEditorMode::ChangeInputObjectLayer(int32 AssetID, int32 NewLayerIndex, bool bForceRebuild)
 {
 	using namespace UVEditorModeLocals;
 
@@ -592,7 +641,7 @@ void UUVEditorMode::ChangeInputObjectLayer(int32 AssetID, int32 NewLayerIndex)
 		{
 			bFound = true;
 
-			if (ensure(InputObject->UVLayerIndex != NewLayerIndex))
+			if (InputObject->UVLayerIndex != NewLayerIndex || bForceRebuild)
 			{
 				InputObject->UVLayerIndex = NewLayerIndex;
 				InputObject->UpdateAllFromAppliedCanonical();
@@ -604,4 +653,24 @@ void UUVEditorMode::ChangeInputObjectLayer(int32 AssetID, int32 NewLayerIndex)
 	ensure(bFound);
 }
 
+void UUVEditorMode::ForceUpdateDisplayChannel(const TArray<int32>& LayerPerAsset, bool bForceRebuildUnwrap, bool bEmitUndoTransaction)
+{
+	for (int32 AssetID = 0; AssetID < ToolInputObjects.Num(); ++AssetID)
+	{
+		if (ToolInputObjects[AssetID]->UVLayerIndex != LayerPerAsset[AssetID] || bForceRebuildUnwrap)
+		{			
+			if (bEmitUndoTransaction)
+			{
+				GetInteractiveToolsContext()->GetTransactionAPI()->AppendChange(this,
+					MakeUnique<UVEditorModeLocals::FInputObjectUVLayerChange>(AssetID, ToolInputObjects[AssetID]->UVLayerIndex, LayerPerAsset[AssetID]),
+					UVEditorModeLocals::UVLayerChangeTransactionName);
+			}
+
+			ChangeInputObjectLayer(AssetID, LayerPerAsset[AssetID], true);
+			PreviousUVLayerIndex[AssetID] = PendingUVLayerIndex[AssetID] = LayerPerAsset[AssetID];
+		}
+	}
+	UpdateSelectedLayer();
+}
+
 #undef LOCTEXT_NAMESPACE
diff --git a/Engine/Plugins/Experimental/UVEditor/Source/UVEditor/Private/UVEditorModeToolkit.cpp b/Engine/Plugins/Experimental/UVEditor/Source/UVEditor/Private/UVEditorModeToolkit.cpp
--- a/Engine/Plugins/Experimental/UVEditor/Source/UVEditor/Private/UVEditorModeToolkit.cpp
+++ b/Engine/Plugins/Experimental/UVEditor/Source/UVEditor/Private/UVEditorModeToolkit.cpp
@@ -188,6 +188,7 @@ void FUVEditorModeToolkit::Init(const TSharedPtr<IToolkitHost>& InitToolkitHost,
 	ToolbarBuilder.AddToolBarButton(CommandInfos.BeginSelectTool);
 	ToolbarBuilder.AddToolBarButton(CommandInfos.BeginLayoutTool);
 	ToolbarBuilder.AddToolBarButton(CommandInfos.BeginParameterizeMeshTool);
+	ToolbarBuilder.AddToolBarButton(CommandInfos.BeginChannelEditTool);
 
 	// Hook in the tool palette
 	ToolButtonsContainer->SetContent(ToolbarBuilder.MakeWidget());
diff --git a/Engine/Plugins/Experimental/UVEditor/Source/UVEditor/Private/UVEditorToolkit.cpp b/Engine/Plugins/Experimental/UVEditor/Source/UVEditor/Private/UVEditorToolkit.cpp
--- a/Engine/Plugins/Experimental/UVEditor/Source/UVEditor/Private/UVEditorToolkit.cpp
+++ b/Engine/Plugins/Experimental/UVEditor/Source/UVEditor/Private/UVEditorToolkit.cpp
@@ -332,6 +332,9 @@ void FUVEditorToolkit::PostInitAssetEditor()
 	EmitChangeAPI->Initialize(EditorModeManager->GetInteractiveToolsContext()->ToolManager);
 	ContextStore->AddContextObject(EmitChangeAPI);
 
+	UUVToolAssetAndChannelAPI* AssetAndLayerAPI = NewObject<UUVToolAssetAndChannelAPI>();
+	ContextStore->AddContextObject(AssetAndLayerAPI);
+
 	// Initialize mode state.
 	TArray<TObjectPtr<UObject>> ObjectsToEdit;
 	OwningAssetEditor->GetObjectsToEdit(ObjectsToEdit);
diff --git a/Engine/Plugins/Experimental/UVEditor/Source/UVEditor/Public/UVEditorCommands.h b/Engine/Plugins/Experimental/UVEditor/Source/UVEditor/Public/UVEditorCommands.h
--- a/Engine/Plugins/Experimental/UVEditor/Source/UVEditor/Public/UVEditorCommands.h
+++ b/Engine/Plugins/Experimental/UVEditor/Source/UVEditor/Public/UVEditorCommands.h
@@ -21,6 +21,7 @@ public:
 	TSharedPtr<FUICommandInfo> BeginSelectTool;
 	TSharedPtr<FUICommandInfo> BeginLayoutTool;
 	TSharedPtr<FUICommandInfo> BeginParameterizeMeshTool;
+	TSharedPtr<FUICommandInfo> BeginChannelEditTool;
 
 	TSharedPtr<FUICommandInfo> AcceptActiveTool;
 	TSharedPtr<FUICommandInfo> CancelActiveTool;
diff --git a/Engine/Plugins/Experimental/UVEditor/Source/UVEditor/Public/UVEditorMode.h b/Engine/Plugins/Experimental/UVEditor/Source/UVEditor/Public/UVEditorMode.h
--- a/Engine/Plugins/Experimental/UVEditor/Source/UVEditor/Public/UVEditorMode.h
+++ b/Engine/Plugins/Experimental/UVEditor/Source/UVEditor/Public/UVEditorMode.h
@@ -57,7 +57,7 @@ public:
 	void InitializeTargets(const TArray<TObjectPtr<UObject>>& AssetsIn, const TArray<FTransform>& TransformsIn);
 
 	// public for use by undo/redo
-	void ChangeInputObjectLayer(int32 AssetID, int32 NewLayerIndex);
+	void ChangeInputObjectLayer(int32 AssetID, int32 NewLayerIndex, bool bForceRebuild=false);
 	void UpdateSelectedLayer();
 
 	bool IsActive() { return bIsActive; }
@@ -177,11 +177,16 @@ protected:
 	void SwitchActiveAsset(const FString& UVAsset);
 	void SwitchActiveChannel(const FString& UVChannel);
 
+	// Used with the ToolAssetAndLayerAPI to process tool layer change requests
+	void ForceUpdateDisplayChannel(const TArray<int32>& LayerPerAsset, bool bForceRebuildUnwrap, bool bEmitUndoTransaction);
+
 	// Used to change layers with our current picker approach (we need to remember the previous
 	// layer value so we know which one to remove). Likely to change if we start adding/removing
 	// using some different UI.
-	int32 PreviousUVLayerIndex = -1;
-	int32 PendingUVLayerIndex = PreviousUVLayerIndex;
+	TArray<int32> PreviousUVLayerIndex;
+	TArray<int32> PendingUVLayerIndex;
+	bool  bPendingUVLayerChangeDestroyTool = true;
+	bool  bForceRebuildUVLayer = false;
 
 
 	// Wireframe Display Properties
diff --git a/Engine/Plugins/Experimental/UVEditor/Source/UVEditorTools/Private/UVEditorLayerEditTool.cpp b/Engine/Plugins/Experimental/UVEditor/Source/UVEditorTools/Private/UVEditorLayerEditTool.cpp
new file mode 100644
--- /dev/null
+++ b/Engine/Plugins/Experimental/UVEditor/Source/UVEditorTools/Private/UVEditorLayerEditTool.cpp
@@ -0,0 +1,651 @@
+// Copyright Epic Games, Inc. All Rights Reserved.
+
+#include "UVEditorLayerEditTool.h"
+
+#include "ContextObjectStore.h"
+#include "DynamicMesh/DynamicMesh3.h"
+#include "DynamicMesh/DynamicMeshChangeTracker.h"
+#include "Parameterization/DynamicMeshUVEditor.h"
+#include "InteractiveToolManager.h"
+#include "MeshOpPreviewHelpers.h" // UMeshOpPreviewWithBackgroundCompute
+#include "ParameterizationOps/UVLayoutOp.h"
+#include "Properties/UVLayoutProperties.h"
+#include "ToolTargets/UVEditorToolMeshInput.h"
+#include "UVToolContextObjects.h"
+#include "ModelingToolTargetUtil.h"
+#include "Components.h"
+
+using namespace UE::Geometry;
+
+#define LOCTEXT_NAMESPACE "UUVChannelEditTool"
+
+namespace UUVEditorChannelEditLocals
+{
+	
+	const FText UVChannelAddTransactionName = LOCTEXT("UVChannelAddTransactionName", "Add UV Channel");
+	const FText UVChannelCloneTransactionName = LOCTEXT("UVChannelCloneTransactionName", "Clone UV Channel");
+	const FText UVChannelDeleteTransactionName = LOCTEXT("UVChannelDeleteTransactionName", "Delete UV Channel");
+
+	class FInputObjectUVChannelAdd : public FToolCommandChange
+	{
+	public:
+		FInputObjectUVChannelAdd(TObjectPtr<UUVEditorToolMeshInput>& TargetIn, int32 AddedUVChannelIndexIn)
+			: Target(TargetIn)
+			, AddedUVChannelIndex(AddedUVChannelIndexIn)			
+		{
+		}
+
+		virtual void Apply(UObject* Object) override
+		{
+			UInteractiveToolManager* InteractiveToolManager = Cast<UInteractiveToolManager>(Object);
+			UUVToolAssetAndChannelAPI* AssetAndChannelAPI = InteractiveToolManager->GetContextObjectStore()->FindContext<UUVToolAssetAndChannelAPI>();
+
+			FDynamicMeshUVEditor DynamicMeshUVEditor(Target->AppliedCanonical.Get(), AddedUVChannelIndex-1, false);
+			int32 NewChannelIndex = DynamicMeshUVEditor.AddUVLayer();
+			check(NewChannelIndex == AddedUVChannelIndex);
+
+			if (NewChannelIndex != -1)
+			{
+				Target->AppliedPreview->PreviewMesh->UpdatePreview(Target->AppliedCanonical.Get());
+
+				AssetAndChannelAPI->NotifyOfAssetChannelCountChange(Target->AssetID);
+				Target->OnUndoRedo.Broadcast(false);
+			}
+
+		}
+
+		virtual void Revert(UObject* Object) override
+		{
+			UInteractiveToolManager* InteractiveToolManager = Cast<UInteractiveToolManager>(Object);
+			UUVToolAssetAndChannelAPI* AssetAndChannelAPI = InteractiveToolManager->GetContextObjectStore()->FindContext<UUVToolAssetAndChannelAPI>();
+
+			FDynamicMeshUVEditor DynamicMeshUVEditor(Target->AppliedCanonical.Get(), AddedUVChannelIndex, false);
+			int32 NewChannelIndex = DynamicMeshUVEditor.RemoveUVLayer();
+			Target->AppliedPreview->PreviewMesh->UpdatePreview(Target->AppliedCanonical.Get());
+			AssetAndChannelAPI->NotifyOfAssetChannelCountChange(Target->AssetID);
+			Target->OnUndoRedo.Broadcast(true);		
+		}
+
+		virtual bool HasExpired(UObject* Object) const override
+		{
+			return false;
+		}
+
+		virtual FString ToString() const override
+		{
+			return TEXT("UVEditorModeLocals::FInputObjectUVChannelAdd");
+		}
+
+	protected:
+		TObjectPtr<UUVEditorToolMeshInput> Target;
+		int32 AddedUVChannelIndex;		
+	};
+
+	class FInputObjectUVChannelClone : public FToolCommandChange
+	{
+	public:
+		FInputObjectUVChannelClone(const TObjectPtr<UUVEditorToolMeshInput>& TargetIn, int32 SourceUVChannelIndexIn, int32 TargetUVChannelIndexIn, const FDynamicMeshUVOverlay& OriginalUVChannelIn)
+			: Target(TargetIn)			
+			, SourceUVChannelIndex(SourceUVChannelIndexIn)
+			, TargetUVChannelIndex(TargetUVChannelIndexIn)
+			, OriginalUVChannel(OriginalUVChannelIn)
+		{
+		}
+
+		virtual void Apply(UObject* Object) override
+		{
+			UInteractiveToolManager* InteractiveToolManager = Cast<UInteractiveToolManager>(Object);
+			UUVToolAssetAndChannelAPI* AssetAndChannelAPI = InteractiveToolManager->GetContextObjectStore()->FindContext<UUVToolAssetAndChannelAPI>();
+
+			FDynamicMeshUVEditor DynamicMeshUVEditor(Target->AppliedCanonical.Get(), TargetUVChannelIndex, false);
+			DynamicMeshUVEditor.CopyUVLayer(Target->AppliedCanonical->Attributes()->GetUVLayer(SourceUVChannelIndex));
+			Target->UpdateAllFromAppliedCanonical();
+			Target->OnUndoRedo.Broadcast(false);
+		}
+
+		virtual void Revert(UObject* Object) override
+		{
+			UInteractiveToolManager* InteractiveToolManager = Cast<UInteractiveToolManager>(Object);
+			UUVToolAssetAndChannelAPI* AssetAndChannelAPI = InteractiveToolManager->GetContextObjectStore()->FindContext<UUVToolAssetAndChannelAPI>();
+
+			FDynamicMeshUVEditor DynamicMeshUVEditor(Target->AppliedCanonical.Get(), TargetUVChannelIndex, false);
+			DynamicMeshUVEditor.CopyUVLayer(&OriginalUVChannel);
+			Target->UpdateAllFromAppliedCanonical();
+			Target->OnUndoRedo.Broadcast(true);
+		}
+
+		virtual bool HasExpired(UObject* Object) const override
+		{
+			return false;
+		}
+
+		virtual FString ToString() const override
+		{
+			return TEXT("UVEditorModeLocals::FInputObjectUVChannelClone");
+		}
+
+	protected:
+		TObjectPtr<UUVEditorToolMeshInput> Target;
+		int32 SourceUVChannelIndex;
+		int32 TargetUVChannelIndex;
+		FDynamicMeshUVOverlay OriginalUVChannel;
+	};
+
+
+	class FInputObjectUVChannelDelete : public FToolCommandChange
+	{
+	public:
+		FInputObjectUVChannelDelete(const TObjectPtr<UUVEditorToolMeshInput>& TargetIn, int32 DeletedUVChannelIndexIn, const FDynamicMeshUVOverlay& OriginalUVChannelIn)
+			: Target(TargetIn)			
+			, DeletedUVChannelIndex(DeletedUVChannelIndexIn)
+			, OriginalUVChannel(OriginalUVChannelIn)
+		{
+		}
+
+		virtual void Apply(UObject* Object) override
+		{
+			UInteractiveToolManager* InteractiveToolManager = Cast<UInteractiveToolManager>(Object);
+			UUVToolAssetAndChannelAPI* AssetAndChannelAPI = InteractiveToolManager->GetContextObjectStore()->FindContext<UUVToolAssetAndChannelAPI>();
+
+			FDynamicMeshUVEditor DynamicMeshUVEditor(Target->AppliedCanonical.Get(), DeletedUVChannelIndex, false);
+
+			DynamicMeshUVEditor.SetPerTriangleUVs(0.0, nullptr); // "clear" the UVs to a single point.
+			int32 NewChannelIndex = DynamicMeshUVEditor.RemoveUVLayer();
+
+			Target->AppliedPreview->PreviewMesh->UpdatePreview(Target->AppliedCanonical.Get());
+
+			AssetAndChannelAPI->NotifyOfAssetChannelCountChange(Target->AssetID);
+			TArray<int32> ChannelPerAsset = AssetAndChannelAPI->GetCurrentChannelVisibility();
+			ChannelPerAsset[Target->AssetID] = DeletedUVChannelIndex;
+			AssetAndChannelAPI->RequestChannelVisibilityChange(ChannelPerAsset, false, false);
+
+			Target->OnUndoRedo.Broadcast(false);
+		}
+
+		virtual void Revert(UObject* Object) override
+		{
+			UInteractiveToolManager* InteractiveToolManager = Cast<UInteractiveToolManager>(Object);
+			UUVToolAssetAndChannelAPI* AssetAndChannelAPI = InteractiveToolManager->GetContextObjectStore()->FindContext<UUVToolAssetAndChannelAPI>();
+
+			FDynamicMeshUVEditor DynamicMeshUVEditor(Target->AppliedCanonical.Get(), DeletedUVChannelIndex, false);
+			int32 NewChannelIndex = DynamicMeshUVEditor.AddUVLayer();
+			Target->AppliedPreview->PreviewMesh->UpdatePreview(Target->AppliedCanonical.Get());
+
+			for (int32 ChannelIndex = NewChannelIndex; ChannelIndex > DeletedUVChannelIndex; --ChannelIndex)
+			{
+				DynamicMeshUVEditor.SwitchActiveLayer(ChannelIndex);
+				DynamicMeshUVEditor.CopyUVLayer(Target->AppliedCanonical->Attributes()->GetUVLayer(ChannelIndex - 1));
+			}
+			DynamicMeshUVEditor.SwitchActiveLayer(DeletedUVChannelIndex);
+			DynamicMeshUVEditor.CopyUVLayer(&OriginalUVChannel);
+
+			AssetAndChannelAPI->NotifyOfAssetChannelCountChange(Target->AssetID);
+			TArray<int32> ChannelPerAsset = AssetAndChannelAPI->GetCurrentChannelVisibility();
+			ChannelPerAsset[Target->AssetID] = DeletedUVChannelIndex;
+			AssetAndChannelAPI->RequestChannelVisibilityChange(ChannelPerAsset, true, false);
+
+			Target->OnUndoRedo.Broadcast(true);
+		}
+
+		virtual bool HasExpired(UObject* Object) const override
+		{
+			return false;
+		}
+
+		virtual FString ToString() const override
+		{
+			return TEXT("UVEditorModeLocals::FInputObjectUVChannelDelete");
+		}
+
+	protected:
+		TObjectPtr<UUVEditorToolMeshInput> Target;		
+		int32 DeletedUVChannelIndex;
+		FDynamicMeshUVOverlay OriginalUVChannel;
+	};
+
+}
+
+
+const FToolTargetTypeRequirements& UUVEditorChannelEditToolBuilder::GetTargetRequirements() const
+{
+	static FToolTargetTypeRequirements TypeRequirements(UUVEditorToolMeshInput::StaticClass());
+	return TypeRequirements;
+}
+
+bool UUVEditorChannelEditToolBuilder::CanBuildTool(const FToolBuilderState& SceneState) const
+{
+	return Targets && Targets->Num() > 0;
+}
+
+UInteractiveTool* UUVEditorChannelEditToolBuilder::BuildTool(const FToolBuilderState& SceneState) const
+{
+	UUVEditorChannelEditTool* NewTool = NewObject<UUVEditorChannelEditTool>(SceneState.ToolManager);
+	NewTool->SetTargets(*Targets);
+
+	return NewTool;
+}
+
+void UUVEditorChannelEditTargetProperties::Initialize(
+	const TArray<TObjectPtr<UUVEditorToolMeshInput>>& TargetsIn,
+	bool bInitializeSelection)
+{
+	UVAssetNames.Reset(TargetsIn.Num());
+	NumUVChannelsPerAsset.Reset(TargetsIn.Num());
+
+	for (int i = 0; i < TargetsIn.Num(); ++i)
+	{
+		UVAssetNames.Add(UE::ToolTarget::GetHumanReadableName(TargetsIn[i]->SourceTarget));
+
+		int32 NumUVChannels = TargetsIn[i]->AppliedCanonical->HasAttributes() ?
+			TargetsIn[i]->AppliedCanonical->Attributes()->NumUVLayers() : 0;
+
+		NumUVChannelsPerAsset.Add(NumUVChannels);
+	}
+
+	if (bInitializeSelection)
+	{
+		Asset = UVAssetNames[0];
+		GetUVChannelNames();
+		TargetChannel = UVChannelNames.Num() > 0 ? UVChannelNames[0] : TEXT("");
+		ReferenceChannel = UVChannelNames.Num() > 0 ? UVChannelNames[0] : TEXT("");
+	}
+}
+
+const TArray<FString>& UUVEditorChannelEditTargetProperties::GetAssetNames()
+{
+	return UVAssetNames;
+}
+
+const TArray<FString>& UUVEditorChannelEditTargetProperties::GetUVChannelNames()
+{
+	int32 FoundAssetIndex = UVAssetNames.IndexOfByKey(Asset);
+
+	if (FoundAssetIndex == INDEX_NONE)
+	{
+		UVChannelNames.Reset();
+		return UVChannelNames;
+	}
+
+	if (UVChannelNames.Num() != NumUVChannelsPerAsset[FoundAssetIndex])
+	{
+		UVChannelNames.Reset();
+		for (int32 i = 0; i < NumUVChannelsPerAsset[FoundAssetIndex]; ++i)
+		{
+			UVChannelNames.Add(FString::Printf(TEXT("UV%d"), i));
+		}
+	}
+
+	return UVChannelNames;
+}
+
+bool UUVEditorChannelEditTargetProperties::ValidateUVAssetSelection(bool bUpdateIfInvalid)
+{
+	int32 FoundIndex = UVAssetNames.IndexOfByKey(Asset);
+	if (FoundIndex == INDEX_NONE)
+	{
+		if (bUpdateIfInvalid)
+		{
+			Asset = (UVAssetNames.Num() > 0) ? UVAssetNames[0] : TEXT("");
+		}
+		return false;
+	}
+	return true;
+}
+
+bool UUVEditorChannelEditTargetProperties::ValidateUVChannelSelection(bool bUpdateIfInvalid)
+{
+	bool bValid = ValidateUVAssetSelection(bUpdateIfInvalid);
+
+	int32 FoundAssetIndex = UVAssetNames.IndexOfByKey(Asset);
+	if (FoundAssetIndex == INDEX_NONE)
+	{
+		if (bUpdateIfInvalid)
+		{
+			TargetChannel = TEXT("");
+			ReferenceChannel = TEXT("");
+		}
+		return false;
+	}
+
+	auto CheckAndUpdateChannel = [this, FoundAssetIndex, bUpdateIfInvalid](FString& UVChannel)
+	{
+		int32 FoundIndex = UVChannelNames.IndexOfByKey(UVChannel);
+		if (FoundIndex >= NumUVChannelsPerAsset[FoundAssetIndex])
+		{
+			FoundIndex = INDEX_NONE;
+		}
+
+		if (FoundIndex == INDEX_NONE)
+		{
+			if (bUpdateIfInvalid)
+			{
+				UVChannel = (NumUVChannelsPerAsset[FoundAssetIndex] > 0) ? UVChannelNames[0] : TEXT("");
+			}
+			return false;
+		}
+		return true;
+	};
+
+	bValid = CheckAndUpdateChannel(TargetChannel);
+	bValid = bValid && CheckAndUpdateChannel(TargetChannel);
+
+	return bValid;
+}
+
+int32 UUVEditorChannelEditTargetProperties::GetSelectedAssetID()
+{
+	int32 FoundIndex = UVAssetNames.IndexOfByKey(Asset);
+	if (FoundIndex == INDEX_NONE)
+	{
+		return INDEX_NONE;
+	}
+	return FoundIndex;
+}
+
+int32 UUVEditorChannelEditTargetProperties::GetSelectedChannelIndex(bool bForceToZeroOnFailure, bool bUseReference)
+{
+	int32 FoundAssetIndex = UVAssetNames.IndexOfByKey(Asset);
+	if (FoundAssetIndex == INDEX_NONE)
+	{
+		return INDEX_NONE;
+	}
+
+	FString UVChannel = bUseReference ? ReferenceChannel : TargetChannel;
+
+	int32 FoundUVIndex = UVChannelNames.IndexOfByKey(UVChannel);
+	if (!ensure(FoundUVIndex < NumUVChannelsPerAsset[FoundAssetIndex]))
+	{
+		return INDEX_NONE;
+	}
+	return FoundUVIndex;
+}
+
+void UUVEditorChannelEditTargetProperties::SetUsageFlags(EChannelEditToolAction Action) {
+	switch (Action)
+	{
+		case EChannelEditToolAction::Add:
+			bActionNeedsAsset = true;
+			bActionNeedsReference = false;
+			bActionNeedsTarget = false;
+			break;
+		case EChannelEditToolAction::Copy:
+			bActionNeedsAsset = true;
+			bActionNeedsReference = true;
+			bActionNeedsTarget = true;
+			break;
+		case EChannelEditToolAction::Delete:
+			bActionNeedsAsset = true;
+			bActionNeedsReference = false;
+			bActionNeedsTarget = true;
+			break;
+		default:
+			bActionNeedsAsset = false;
+			bActionNeedsReference = false;
+			bActionNeedsTarget = false;
+			break;
+	}
+}
+
+
+
+
+// Tool property functions
+
+void  UUVEditorChannelEditToolActionPropertySet::Apply()
+{
+	if (ParentTool.IsValid())
+	{
+		PostAction(ParentTool->ActiveAction());
+		return;
+	}	
+	PostAction(EChannelEditToolAction::NoAction);
+}
+
+
+void UUVEditorChannelEditToolActionPropertySet::PostAction(EChannelEditToolAction Action)
+{
+	if (ParentTool.IsValid())
+	{
+		ParentTool->RequestAction(Action);
+	}
+}
+
+
+
+void UUVEditorChannelEditTool::Setup()
+{
+	check(Targets.Num() > 0);
+
+	UInteractiveTool::Setup();
+
+	EmitChangeAPI = GetToolManager()->GetContextObjectStore()->FindContext<UUVToolEmitChangeAPI>();
+
+	ActionSelectionProperties = NewObject<UUVEditorChannelEditSettings>(this);
+	ActionSelectionProperties->WatchProperty(ActionSelectionProperties->Action, [this](EChannelEditToolAction Action) {
+		if (SourceChannelProperties)
+		{
+			SourceChannelProperties->SetUsageFlags(Action);
+			NotifyOfPropertyChangeByTool(SourceChannelProperties);
+		}
+	});
+	AddToolPropertySource(ActionSelectionProperties);
+
+	SourceChannelProperties = NewObject<UUVEditorChannelEditTargetProperties>(this);
+	SourceChannelProperties->Initialize(Targets,true);
+	SourceChannelProperties->WatchProperty(SourceChannelProperties->Asset, [this](FString UVAsset) {ApplyVisbleChannelChange(); });
+	SourceChannelProperties->WatchProperty(SourceChannelProperties->TargetChannel, [this](FString UVChannel) {ApplyVisbleChannelChange(); });
+
+
+	AddToolPropertySource(SourceChannelProperties);
+
+	AddActionProperties = NewObject<UUVEditorChannelEditAddProperties>(this);
+	AddToolPropertySource(AddActionProperties);
+
+	CopyActionProperties = NewObject<UUVEditorChannelEditCloneProperties>(this);
+	AddToolPropertySource(CopyActionProperties);
+
+	DeleteActionProperties = NewObject<UUVEditorChannelEditDeleteProperties>(this);
+	AddToolPropertySource(DeleteActionProperties);
+
+	ToolActions = NewObject<UUVEditorChannelEditToolActionPropertySet>(this);
+	ToolActions->Initialize(this);
+	AddToolPropertySource(ToolActions);
+
+	SetToolDisplayName(LOCTEXT("ToolName", "UV Channel Edit"));
+	GetToolManager()->DisplayMessage(LOCTEXT("OnStartUVChannelEditTool", "Add/Copy/Delete Asset UV Channels"),
+		EToolMessageLevel::UserNotification);
+
+	for (int32 i = 0; i < Targets.Num(); ++i)
+	{
+		Targets[i]->OnUndoRedo.AddWeakLambda(this, [this, i](bool bRevert) {
+			UpdateChannelSelectionProperties(i);
+			});
+	}
+}
+
+void UUVEditorChannelEditTool::Shutdown(EToolShutdownType ShutdownType)
+{
+	for (TObjectPtr<UUVEditorToolMeshInput> Target : Targets)
+	{
+		Target->OnUndoRedo.RemoveAll(this);
+	}
+
+	Targets.Empty();
+}
+
+void UUVEditorChannelEditTool::UpdateChannelSelectionProperties(int32 ChangingAsset)
+{
+	SourceChannelProperties->Initialize(Targets, false);
+
+	UUVToolAssetAndChannelAPI* AssetAndChannelAPI = GetToolManager()->GetContextObjectStore()->FindContext<UUVToolAssetAndChannelAPI>();
+	if (AssetAndChannelAPI)
+	{
+		ActiveChannel = AssetAndChannelAPI->GetCurrentChannelVisibility()[ChangingAsset];
+
+		SourceChannelProperties->Asset = SourceChannelProperties->UVAssetNames[ChangingAsset];
+		SourceChannelProperties->TargetChannel = SourceChannelProperties->GetUVChannelNames()[ActiveChannel];
+		SourceChannelProperties->ValidateUVChannelSelection(true);
+		SourceChannelProperties->SilentUpdateWatched();
+	}
+}
+
+void UUVEditorChannelEditTool::OnTick(float DeltaTime)
+{
+	if (PendingAction != EChannelEditToolAction::NoAction)
+	{
+		ActiveAsset = SourceChannelProperties->GetSelectedAssetID();
+		ActiveChannel = SourceChannelProperties->GetSelectedChannelIndex(true, false);
+		ReferenceChannel = SourceChannelProperties->GetSelectedChannelIndex(true, true);
+
+		switch (PendingAction)
+		{
+		case EChannelEditToolAction::Add:
+		{
+			AddChannel();
+		}
+		break;
+
+		case EChannelEditToolAction::Copy:
+		{
+			CopyChannel();
+		}
+		break;
+
+		case EChannelEditToolAction::Delete:
+		{
+			DeleteChannel();
+		}
+		break;
+
+		default:
+			checkSlow(false);
+		}
+
+		PendingAction = EChannelEditToolAction::NoAction;
+	}
+}
+
+void UUVEditorChannelEditTool::ApplyVisbleChannelChange()
+{
+	SourceChannelProperties->ValidateUVAssetSelection(true);
+	SourceChannelProperties->ValidateUVChannelSelection(true);
+
+	ActiveAsset = SourceChannelProperties->GetSelectedAssetID();
+	ActiveChannel = SourceChannelProperties->GetSelectedChannelIndex(true, false);
+
+	UUVToolAssetAndChannelAPI* AssetAndChannelAPI = GetToolManager()->GetContextObjectStore()->FindContext<UUVToolAssetAndChannelAPI>();
+	if (AssetAndChannelAPI)
+	{
+		TArray<int32> ChannelPerAsset = AssetAndChannelAPI->GetCurrentChannelVisibility();
+		ChannelPerAsset[ActiveAsset] = ActiveChannel;
+		AssetAndChannelAPI->RequestChannelVisibilityChange(ChannelPerAsset);
+	}
+}
+
+bool UUVEditorChannelEditTool::CanAccept() const
+{
+	return true;
+}
+
+void UUVEditorChannelEditTool::OnPropertyModified(UObject* PropertySet, FProperty* Property)
+{
+
+}
+
+EChannelEditToolAction UUVEditorChannelEditTool::ActiveAction() const
+{
+	if (ActionSelectionProperties)
+	{
+		return ActionSelectionProperties->Action;
+	}
+	return EChannelEditToolAction::NoAction;
+}
+
+void UUVEditorChannelEditTool::RequestAction(EChannelEditToolAction ActionType)
+{
+	PendingAction = ActionType;
+}
+
+void UUVEditorChannelEditTool::AddChannel()
+{
+	UUVToolAssetAndChannelAPI* AssetAndChannelAPI = GetToolManager()->GetContextObjectStore()->FindContext<UUVToolAssetAndChannelAPI>();
+	FDynamicMeshUVEditor DynamicMeshUVEditor(Targets[ActiveAsset]->AppliedCanonical.Get(), ActiveChannel, false);
+	int32 NewChannelIndex = DynamicMeshUVEditor.AddUVLayer();
+
+	if (NewChannelIndex != -1)
+	{
+		Targets[ActiveAsset]->AppliedPreview->PreviewMesh->UpdatePreview(Targets[ActiveAsset]->AppliedCanonical.Get());
+
+		EmitChangeAPI->BeginUndoTransaction(UUVEditorChannelEditLocals::UVChannelAddTransactionName);
+		EmitChangeAPI->EmitToolIndependentChange(GetToolManager(),
+			MakeUnique<UUVEditorChannelEditLocals::FInputObjectUVChannelAdd>(Targets[ActiveAsset], NewChannelIndex),
+			UUVEditorChannelEditLocals::UVChannelAddTransactionName);
+
+		AssetAndChannelAPI->NotifyOfAssetChannelCountChange(ActiveAsset);
+		TArray<int32> ChannelPerAsset = AssetAndChannelAPI->GetCurrentChannelVisibility();
+		ChannelPerAsset[ActiveAsset] = NewChannelIndex;
+		AssetAndChannelAPI->RequestChannelVisibilityChange(ChannelPerAsset);
+
+		EmitChangeAPI->EndUndoTransaction();
+
+		SourceChannelProperties->Initialize(Targets, false);
+		SourceChannelProperties->TargetChannel = SourceChannelProperties->GetUVChannelNames()[NewChannelIndex];
+	}
+
+}
+
+void UUVEditorChannelEditTool::CopyChannel()
+{
+	UUVToolAssetAndChannelAPI* AssetAndChannelAPI = GetToolManager()->GetContextObjectStore()->FindContext<UUVToolAssetAndChannelAPI>();
+
+	FDynamicMeshUVOverlay OriginalUVChannel(Targets[ActiveAsset]->AppliedCanonical.Get());
+	OriginalUVChannel.Copy(*Targets[ActiveAsset]->AppliedCanonical->Attributes()->GetUVLayer(ActiveChannel));
+
+	FDynamicMeshUVEditor DynamicMeshUVEditor(Targets[ActiveAsset]->AppliedCanonical.Get(), ActiveChannel, false);
+	DynamicMeshUVEditor.CopyUVLayer( Targets[ActiveAsset]->AppliedCanonical->Attributes()->GetUVLayer(ReferenceChannel) );
+
+	EmitChangeAPI->BeginUndoTransaction(UUVEditorChannelEditLocals::UVChannelCloneTransactionName);
+	EmitChangeAPI->EmitToolIndependentChange(GetToolManager(),
+		MakeUnique<UUVEditorChannelEditLocals::FInputObjectUVChannelClone>(Targets[ActiveAsset], ReferenceChannel, ActiveChannel, OriginalUVChannel),
+		UUVEditorChannelEditLocals::UVChannelCloneTransactionName);
+
+	TArray<int32> ChannelPerAsset = AssetAndChannelAPI->GetCurrentChannelVisibility();
+	ChannelPerAsset[ActiveAsset] = ActiveChannel;
+	AssetAndChannelAPI->RequestChannelVisibilityChange(ChannelPerAsset, true);
+
+	EmitChangeAPI->EndUndoTransaction();
+
+	SourceChannelProperties->Initialize(Targets, false);
+}
+
+void UUVEditorChannelEditTool::DeleteChannel()
+{
+	UUVToolAssetAndChannelAPI* AssetAndChannelAPI = GetToolManager()->GetContextObjectStore()->FindContext<UUVToolAssetAndChannelAPI>();
+	FDynamicMeshUVEditor DynamicMeshUVEditor(Targets[ActiveAsset]->AppliedCanonical.Get(), ActiveChannel, false);
+	
+	FDynamicMeshUVOverlay OriginalUVChannel(Targets[ActiveAsset]->AppliedCanonical.Get());
+	OriginalUVChannel.Copy(*Targets[ActiveAsset]->AppliedCanonical->Attributes()->GetUVLayer(ActiveChannel));
+
+	DynamicMeshUVEditor.SetPerTriangleUVs(0.0, nullptr); // "clear" the UVs to a single point.
+	int32 NewChannelIndex = DynamicMeshUVEditor.RemoveUVLayer();
+	
+	Targets[ActiveAsset]->AppliedPreview->PreviewMesh->UpdatePreview(Targets[ActiveAsset]->AppliedCanonical.Get());
+
+	EmitChangeAPI->BeginUndoTransaction(UUVEditorChannelEditLocals::UVChannelDeleteTransactionName);
+	EmitChangeAPI->EmitToolIndependentChange(GetToolManager(),
+		MakeUnique<UUVEditorChannelEditLocals::FInputObjectUVChannelDelete>(Targets[ActiveAsset], ActiveChannel, OriginalUVChannel),
+		UUVEditorChannelEditLocals::UVChannelDeleteTransactionName);
+
+	AssetAndChannelAPI->NotifyOfAssetChannelCountChange(ActiveAsset);
+	TArray<int32> ChannelPerAsset = AssetAndChannelAPI->GetCurrentChannelVisibility();
+	ChannelPerAsset[ActiveAsset] = NewChannelIndex;
+	AssetAndChannelAPI->RequestChannelVisibilityChange(ChannelPerAsset, true);
+
+	EmitChangeAPI->EndUndoTransaction();
+
+
+	SourceChannelProperties->Initialize(Targets, false);
+	SourceChannelProperties->TargetChannel = SourceChannelProperties->GetUVChannelNames()[NewChannelIndex];
+}
+
+#undef LOCTEXT_NAMESPACE
\ No newline at end of file
diff --git a/Engine/Plugins/Experimental/UVEditor/Source/UVEditorTools/Private/UVEditorLayoutTool.cpp b/Engine/Plugins/Experimental/UVEditor/Source/UVEditorTools/Private/UVEditorLayoutTool.cpp
--- a/Engine/Plugins/Experimental/UVEditor/Source/UVEditorTools/Private/UVEditorLayoutTool.cpp
+++ b/Engine/Plugins/Experimental/UVEditor/Source/UVEditorTools/Private/UVEditorLayoutTool.cpp
@@ -14,7 +14,7 @@
 
 using namespace UE::Geometry;
 
-#define LOCTEXT_NAMESPACE "UUVSelectTool"
+#define LOCTEXT_NAMESPACE "UUVEditorLayoutTool"
 
 // Tool builder
 // TODO: Could consider sharing some of the tool builder boilerplate for UV editor tools in a common base class.
diff --git a/Engine/Plugins/Experimental/UVEditor/Source/UVEditorTools/Public/UVEditorLayerEditTool.h b/Engine/Plugins/Experimental/UVEditor/Source/UVEditorTools/Public/UVEditorLayerEditTool.h
new file mode 100644
--- /dev/null
+++ b/Engine/Plugins/Experimental/UVEditor/Source/UVEditorTools/Public/UVEditorLayerEditTool.h
@@ -0,0 +1,234 @@
+// Copyright Epic Games, Inc. All Rights Reserved.
+
+#pragma once
+
+#include "CoreMinimal.h"
+
+#include "InteractiveTool.h"
+#include "InteractiveToolBuilder.h"
+
+#include "UVEditorLayerEditTool.generated.h"
+
+class UUVEditorToolMeshInput;
+class UUVEditorChannelEditTool;
+class UUVToolEmitChangeAPI;
+
+UCLASS()
+class UVEDITORTOOLS_API UUVEditorChannelEditToolBuilder : public UInteractiveToolBuilder
+{
+	GENERATED_BODY()
+
+public:
+	virtual bool CanBuildTool(const FToolBuilderState& SceneState) const override;
+	virtual UInteractiveTool* BuildTool(const FToolBuilderState& SceneState) const override;
+
+	// This is a pointer so that it can be updated under the builder without
+	// having to set it in the mode after initializing targets.
+	const TArray<TObjectPtr<UUVEditorToolMeshInput>>* Targets = nullptr;
+
+protected:
+	virtual const FToolTargetTypeRequirements& GetTargetRequirements() const override;
+};
+
+UENUM()
+enum class EChannelEditToolAction
+{
+	NoAction,
+
+	Add,
+	Copy,
+	Delete
+};
+
+
+UCLASS()
+class UVEDITORTOOLS_API UUVEditorChannelEditSettings : public UInteractiveToolPropertySet
+{
+	GENERATED_BODY()
+public:
+	/** UV Layer Edit action to preform */
+	UPROPERTY(EditAnywhere, Category = Options)
+	EChannelEditToolAction Action = EChannelEditToolAction::Add;
+};
+
+UCLASS()
+class UVEDITORTOOLS_API UUVEditorChannelEditTargetProperties : public UInteractiveToolPropertySet
+{
+	GENERATED_BODY()
+public:
+
+	UPROPERTY(EditAnywhere, Category = "UVChannels", meta = (DisplayName = "Asset", GetOptions = GetAssetNames, EditCondition = "bActionNeedsAsset", EditConditionHides = true, HideEditConditionToggle = true))
+	FString Asset;
+
+	UFUNCTION()
+	const TArray<FString>& GetAssetNames();
+
+	UPROPERTY(EditAnywhere, Category = "UVChannels", meta = (DisplayName = "Target UV Channel", GetOptions = GetUVChannelNames, EditCondition = "bActionNeedsTarget", EditConditionHides = true, HideEditConditionToggle = true))
+	FString TargetChannel;
+
+	UPROPERTY(EditAnywhere, Category = "UVChannels", meta = (DisplayName = "Source UV Channel", GetOptions = GetUVChannelNames, EditCondition = "bActionNeedsReference", EditConditionHides = true, HideEditConditionToggle = true))
+	FString ReferenceChannel;
+
+	UFUNCTION()
+	const TArray<FString>& GetUVChannelNames();
+
+	TArray<FString> UVChannelNames;
+	TArray<FString> UVAssetNames;
+
+	// 1:1 with UVAssetNames
+	TArray<int32> NumUVChannelsPerAsset;
+
+	UPROPERTY(meta = (TransientToolProperty))
+	bool bActionNeedsAsset = true;
+
+	UPROPERTY(meta = (TransientToolProperty))
+	bool bActionNeedsReference = false;
+
+	UPROPERTY(meta = (TransientToolProperty))
+	bool bActionNeedsTarget = true;
+
+public:
+	void Initialize(
+		const TArray<TObjectPtr<UUVEditorToolMeshInput>>& TargetsIn,
+		bool bInitializeSelection);
+
+	void SetUsageFlags(EChannelEditToolAction Action);
+
+	/**
+	 * Verify that the UV asset selection is valid
+	 * @param bUpdateIfInvalid if selection is not valid, reset UVAsset to Asset0 or empty if no assets exist
+	 * @return true if selection in UVAsset is an entry in UVAssetNames.
+	 */
+	bool ValidateUVAssetSelection(bool bUpdateIfInvalid);
+
+	/**
+	 * Verify that the UV channel selection is valid
+	 * @param bUpdateIfInvalid if selection is not valid, UVChannel to UV0 or empty if no UV channels exist
+	 * @return true if selection in UVChannel is an entry in UVChannelNamesList.
+	 */
+	bool ValidateUVChannelSelection(bool bUpdateIfInvalid);
+
+
+	/**
+	 * @return selected UV asset ID, or -1 if invalid selection
+	 */
+	int32 GetSelectedAssetID();
+
+	/**
+	 * @param bForceToZeroOnFailure if true, then instead of returning -1 we return 0 so calling code can fallback to default UV paths
+	 * @param bUseReference if true, get the selected reference channel index, otherwise return the target channel's index.
+	 * @return selected UV channel index, or -1 if invalid selection, or 0 if invalid selection and bool bForceToZeroOnFailure = true
+	 */
+	int32 GetSelectedChannelIndex(bool bForceToZeroOnFailure = false, bool bUseReference = false);
+};
+
+UCLASS()
+class UVEDITORTOOLS_API UUVEditorChannelEditAddProperties : public UInteractiveToolPropertySet
+{
+	GENERATED_BODY()
+public:
+	/* Placeholder for future per action settings */
+};
+
+UCLASS()
+class UVEDITORTOOLS_API UUVEditorChannelEditCloneProperties : public UInteractiveToolPropertySet
+{
+	GENERATED_BODY()
+public:
+	/* Placeholder for future per action settings */
+};
+
+UCLASS()
+class UVEDITORTOOLS_API UUVEditorChannelEditDeleteProperties : public UInteractiveToolPropertySet
+{
+	GENERATED_BODY()
+public:
+	/* Placeholder for future per action settings */
+};
+
+UCLASS()
+class UVEDITORTOOLS_API UUVEditorChannelEditToolActionPropertySet : public UInteractiveToolPropertySet
+{
+	GENERATED_BODY()
+
+public:
+	TWeakObjectPtr<UUVEditorChannelEditTool> ParentTool;
+
+	void Initialize(UUVEditorChannelEditTool* ParentToolIn) { ParentTool = ParentToolIn; }
+	void PostAction(EChannelEditToolAction Action);
+	
+	UFUNCTION(CallInEditor, Category = Actions)
+	void Apply();
+};
+
+/**
+ *
+ */
+UCLASS()
+class UVEDITORTOOLS_API UUVEditorChannelEditTool : public UInteractiveTool
+{
+	GENERATED_BODY()
+
+public:
+
+	/**
+	 * The tool will operate on the meshes given here.
+	 */
+	virtual void SetTargets(const TArray<TObjectPtr<UUVEditorToolMeshInput>>& TargetsIn)
+	{
+		Targets = TargetsIn;
+	}
+
+	virtual void Setup() override;
+	virtual void Shutdown(EToolShutdownType ShutdownType) override;
+
+	virtual void OnTick(float DeltaTime) override;
+
+	virtual bool HasCancel() const override { return true; }
+	virtual bool HasAccept() const override { return false; }
+	virtual bool CanAccept() const override;
+
+	virtual void OnPropertyModified(UObject* PropertySet, FProperty* Property) override;
+
+	void RequestAction(EChannelEditToolAction ActionType);
+	EChannelEditToolAction ActiveAction() const;
+
+
+protected:
+	void ApplyVisbleChannelChange();
+	void UpdateChannelSelectionProperties(int32 ChangingAsset);
+
+	void AddChannel();
+	void CopyChannel();
+	void DeleteChannel();
+	int32 ActiveAsset;
+	int32 ActiveChannel;
+	int32 ReferenceChannel;
+
+	EChannelEditToolAction PendingAction = EChannelEditToolAction::NoAction;
+
+	UPROPERTY()
+	TArray<TObjectPtr<UUVEditorToolMeshInput>> Targets;
+
+	UPROPERTY()
+	TObjectPtr<UUVEditorChannelEditToolActionPropertySet> ToolActions = nullptr;
+
+	UPROPERTY()
+	TObjectPtr<UUVEditorChannelEditSettings> ActionSelectionProperties = nullptr;
+
+	UPROPERTY()
+	TObjectPtr<UUVEditorChannelEditTargetProperties> SourceChannelProperties = nullptr;
+
+	UPROPERTY()
+	TObjectPtr<UUVEditorChannelEditAddProperties> AddActionProperties = nullptr;
+
+	UPROPERTY()
+	TObjectPtr<UUVEditorChannelEditCloneProperties> CopyActionProperties = nullptr;
+
+	UPROPERTY()
+	TObjectPtr<UUVEditorChannelEditDeleteProperties> DeleteActionProperties = nullptr;
+
+	UPROPERTY()
+	TObjectPtr<UUVToolEmitChangeAPI> EmitChangeAPI = nullptr;
+
+};
diff --git a/Engine/Plugins/Experimental/UVEditor/Source/UVEditorTools/Public/UVToolContextObjects.h b/Engine/Plugins/Experimental/UVEditor/Source/UVEditorTools/Public/UVToolContextObjects.h
--- a/Engine/Plugins/Experimental/UVEditor/Source/UVEditorTools/Public/UVToolContextObjects.h
+++ b/Engine/Plugins/Experimental/UVEditor/Source/UVEditorTools/Public/UVToolContextObjects.h
@@ -149,6 +149,48 @@ protected:
 	EGizmoMode GizmoMode = EGizmoMode::Select;
 };
 
+
+/**
+ * Allows tools to interact with the assets and their UV layers
+*/
+UCLASS()
+class UVEDITORTOOLS_API UUVToolAssetAndChannelAPI : public UUVToolContextObject
+{
+	GENERATED_BODY()
+public:
+
+	TArray<int32> GetCurrentChannelVisibility()
+	{
+		if (GetCurrentChannelVisibilityFunc)
+		{
+			return GetCurrentChannelVisibilityFunc();
+		}
+		return TArray<int32>();
+	}
+
+	void RequestChannelVisibilityChange(const TArray<int32>& ChannelPerAsset, bool bForceRebuildUnwrap=false, bool bEmitUndoTransaction=true)
+	{
+		if (RequestChannelVisibilityChangeFunc)
+		{
+			RequestChannelVisibilityChangeFunc(ChannelPerAsset, bForceRebuildUnwrap, bEmitUndoTransaction);
+		}
+	}
+
+	void NotifyOfAssetChannelCountChange(int32 AssetID)
+	{
+		if (NotifyOfAssetChannelCountChangeFunc)
+		{
+			NotifyOfAssetChannelCountChangeFunc(AssetID);
+		}
+	}
+
+
+	TUniqueFunction<TArray<int32>()> GetCurrentChannelVisibilityFunc;
+	TUniqueFunction<void(const TArray<int32>&, bool, bool)> RequestChannelVisibilityChangeFunc;
+	TUniqueFunction<void(int32 AssetID)> NotifyOfAssetChannelCountChangeFunc;
+
+};
+
 /** Stores UV mesh AABB trees */
 UCLASS()
 class UVEDITORTOOLS_API UUVToolAABBTreeStorage : public UUVToolContextObject
diff --git a/Engine/Plugins/Experimental/UVEditor/Source/UVEditorTools/UVEditorTools.Build.cs b/Engine/Plugins/Experimental/UVEditor/Source/UVEditorTools/UVEditorTools.Build.cs
--- a/Engine/Plugins/Experimental/UVEditor/Source/UVEditorTools/UVEditorTools.Build.cs
+++ b/Engine/Plugins/Experimental/UVEditor/Source/UVEditorTools/UVEditorTools.Build.cs
@@ -49,7 +49,7 @@ public class UVEditorTools : ModuleRules
 				"DynamicMesh",
 				"GeometryCore",
 				"ModelingComponents",
-				"ModelingOperators"
+				"ModelingOperators",				
 				// ... add private dependencies that you statically link with here ...	
 			}
 			);
diff --git a/Engine/Plugins/Runtime/GeometryProcessing/Source/DynamicMesh/Private/Parameterization/DynamicMeshUVEditor.cpp b/Engine/Plugins/Runtime/GeometryProcessing/Source/DynamicMesh/Private/Parameterization/DynamicMeshUVEditor.cpp
--- a/Engine/Plugins/Runtime/GeometryProcessing/Source/DynamicMesh/Private/Parameterization/DynamicMeshUVEditor.cpp
+++ b/Engine/Plugins/Runtime/GeometryProcessing/Source/DynamicMesh/Private/Parameterization/DynamicMeshUVEditor.cpp
@@ -24,19 +24,25 @@
 #include "ExplicitUseGeometryMathTypes.h"		// using UE::Geometry::(math types)
 using namespace UE::Geometry;
 
-FDynamicMeshUVEditor::FDynamicMeshUVEditor(FDynamicMesh3* MeshIn, int32 UVLayerIndex, bool bCreateIfMissing)
+namespace FDynamicMeshUVEditorLocals
+{
+	enum { MAX_TEXCOORDS = 4, MAX_STATIC_TEXCOORDS = 8 };
+}
+
+FDynamicMeshUVEditor::FDynamicMeshUVEditor(FDynamicMesh3* MeshIn, int32 UVLayerIndexIn, bool bCreateIfMissing)
 {
 	Mesh = MeshIn;
+	UVOverlayIndex = UVLayerIndexIn;
 
-	if (Mesh->HasAttributes() && Mesh->Attributes()->NumUVLayers() > UVLayerIndex)
+	if (Mesh->HasAttributes() && Mesh->Attributes()->NumUVLayers() > UVOverlayIndex)
 	{
-		UVOverlay = Mesh->Attributes()->GetUVLayer(UVLayerIndex);
+		UVOverlay = Mesh->Attributes()->GetUVLayer(UVOverlayIndex);
 	}
 
 	if (UVOverlay == nullptr && bCreateIfMissing)
 	{
-		CreateUVLayer(UVLayerIndex);
-		UVOverlay = Mesh->Attributes()->GetUVLayer(UVLayerIndex);
+		CreateUVLayer(UVOverlayIndex);
+		UVOverlay = Mesh->Attributes()->GetUVLayer(UVOverlayIndex);
 		check(UVOverlay);
 	}
 }
@@ -46,12 +52,27 @@ FDynamicMeshUVEditor::FDynamicMeshUVEditor(FDynamicMesh3* MeshIn, FDynamicMeshUV
 {
 	Mesh = MeshIn;
 	UVOverlay = UVOverlayIn;
+	UVOverlayIndex = -1;
+
+	if (Mesh->HasAttributes())
+	{
+		for (int32 UVIndex = 0; UVIndex < Mesh->Attributes()->NumUVLayers(); ++UVIndex) {
+			if (Mesh->Attributes()->GetUVLayer(UVIndex) == UVOverlay)
+			{
+				UVOverlayIndex = UVIndex;
+				break;
+			}
+		}
+	}
+	check(UVOverlayIndex != -1);
 	ensure(UVOverlay->GetParentMesh() == Mesh);
 }
 
 
 void FDynamicMeshUVEditor::CreateUVLayer(int32 LayerIndex)
 {
+	check(LayerIndex < FDynamicMeshUVEditorLocals::MAX_STATIC_TEXCOORDS)
+
 	if (Mesh->HasAttributes() == false)
 	{
 		Mesh->EnableAttributes();
@@ -64,6 +85,52 @@ void FDynamicMeshUVEditor::CreateUVLayer(int32 LayerIndex)
 }
 
 
+int32 FDynamicMeshUVEditor::AddUVLayer()
+{
+	int32 OldUVOverlayIndex = UVOverlayIndex;
+	int32 TotalLayerCount = Mesh->Attributes()->NumUVLayers();
+	if (TotalLayerCount < FDynamicMeshUVEditorLocals::MAX_STATIC_TEXCOORDS) {
+		CreateUVLayer(TotalLayerCount); // The passed in value is an index, not count
+		SwitchActiveLayer(TotalLayerCount);
+		SetPerTriangleUVs(0.0, nullptr);
+		SwitchActiveLayer(OldUVOverlayIndex);
+		return TotalLayerCount;
+	}
+	else
+	{
+		return -1;
+	}
+}
+
+void FDynamicMeshUVEditor::SwitchActiveLayer(int32 UVOverlayIndexIn)
+{
+	UVOverlay = Mesh->Attributes()->GetUVLayer(UVOverlayIndexIn);
+	UVOverlayIndex = UVOverlayIndexIn;
+	check(UVOverlay);
+}
+
+int32 FDynamicMeshUVEditor::RemoveUVLayer()
+{
+	int32 TotalLayerCount = Mesh->Attributes()->NumUVLayers();
+	if (TotalLayerCount == 1)
+	{
+		return 0; // Don't remove the last layer, if there's only one.
+	}
+
+	for (int32 LayerID = UVOverlayIndex + 1; LayerID < TotalLayerCount; ++LayerID)
+	{
+		FDynamicMeshUVOverlay* SourceOverlay = Mesh->Attributes()->GetUVLayer(LayerID);
+		UVOverlay = Mesh->Attributes()->GetUVLayer(LayerID - 1);
+		CopyUVLayer(SourceOverlay);
+	}
+
+	Mesh->Attributes()->SetNumUVLayers(TotalLayerCount-1);	
+	int32 NewLayerIndex = UVOverlayIndex - 1 < 0 ? 0 : UVOverlayIndex - 1;
+	SwitchActiveLayer(NewLayerIndex);
+
+	return NewLayerIndex;
+}
+
 void FDynamicMeshUVEditor::ResetUVs()
 {
 	if (ensure(UVOverlay))
diff --git a/Engine/Plugins/Runtime/GeometryProcessing/Source/DynamicMesh/Public/Parameterization/DynamicMeshUVEditor.h b/Engine/Plugins/Runtime/GeometryProcessing/Source/DynamicMesh/Public/Parameterization/DynamicMeshUVEditor.h
--- a/Engine/Plugins/Runtime/GeometryProcessing/Source/DynamicMesh/Public/Parameterization/DynamicMeshUVEditor.h
+++ b/Engine/Plugins/Runtime/GeometryProcessing/Source/DynamicMesh/Public/Parameterization/DynamicMeshUVEditor.h
@@ -28,6 +28,8 @@ private:
 	FDynamicMesh3* Mesh = nullptr;
 	/** The UV Overlay we will be editing */
 	FDynamicMeshUVOverlay* UVOverlay = nullptr;
+	/** The index of the UV Overlay we will be editing */
+	int32 UVOverlayIndex = -1;
 
 public:
 	/**
@@ -54,6 +56,21 @@ public:
 	 */
 	void CreateUVLayer(int32 UVLayerIndex);
 
+	/**
+	 * Append new UV layer to the end of the array, returning the newly added index, or -1 if at max layers already
+	 */
+	int32 AddUVLayer();
+
+	/**
+	 * Switch to new layer
+	 */
+	void SwitchActiveLayer(int32 UVLayerIndex);
+
+	/**
+	 * Remove active layer, setting active layer to highest, preceeding layer. Will not remove the final layer from the mesh. Returns the resulting active layer index.
+	 */
+	int32 RemoveUVLayer();
+
 	/**
 	 * Clear UVs for all triangles on active layer
 	 */
