UVEditor: Add minimal analytics

#rnx

diff --git a/Engine/Plugins/Experimental/UVEditor/Source/UVEditor/Private/UVEditorMode.cpp b/Engine/Plugins/Experimental/UVEditor/Source/UVEditor/Private/UVEditorMode.cpp
--- a/Engine/Plugins/Experimental/UVEditor/Source/UVEditor/Private/UVEditorMode.cpp
+++ b/Engine/Plugins/Experimental/UVEditor/Source/UVEditor/Private/UVEditorMode.cpp
@@ -7,6 +7,7 @@
 #include "Drawing/MeshElementsVisualizer.h"
 #include "Editor.h"
 #include "EdModeInteractiveToolsContext.h" //ToolsContext
+#include "EngineAnalytics.h"
 #include "Framework/Commands/UICommandList.h"
 #include "InteractiveTool.h"
 #include "MeshOpPreviewHelpers.h" // UMeshOpPreviewWithBackgroundCompute
@@ -39,6 +40,8 @@ using namespace UE::Geometry;
 
 const FEditorModeID UUVEditorMode::EM_UVEditorModeId = TEXT("EM_UVEditorMode");
 
+FDateTime UUVEditorMode::AnalyticsLastStartTimestamp;
+
 namespace UVEditorModeLocals
 {
 	// The layer we open when we first open the UV editor
@@ -144,6 +147,16 @@ void UUVEditorMode::Enter()
 
 	ActivateDefaultTool();
 
+	if (FEngineAnalytics::IsAvailable())
+	{
+		AnalyticsLastStartTimestamp = FDateTime::UtcNow();
+
+		TArray<FAnalyticsEventAttribute> Attributes;
+		Attributes.Add(FAnalyticsEventAttribute(TEXT("Timestamp"), AnalyticsLastStartTimestamp.ToString()));
+
+		FEngineAnalytics::GetProvider().RecordEvent(TEXT("Editor.Usage.UVEditor.Enter"), Attributes);
+	}
+
 	bIsActive = true;
 }
 
@@ -222,6 +235,18 @@ void UUVEditorMode::BindCommands()
 
 void UUVEditorMode::Exit()
 {
+	if (FEngineAnalytics::IsAvailable())
+	{
+		const FDateTime Now = FDateTime::UtcNow();
+		const FTimespan ModeUsageDuration = Now - AnalyticsLastStartTimestamp;
+
+		TArray<FAnalyticsEventAttribute> Attributes;
+		Attributes.Add(FAnalyticsEventAttribute(TEXT("Timestamp"), Now.ToString()));
+		Attributes.Add(FAnalyticsEventAttribute(TEXT("Duration.Seconds"), static_cast<float>(ModeUsageDuration.GetTotalSeconds())));
+
+		FEngineAnalytics::GetProvider().RecordEvent(TEXT("Editor.Usage.UVEditor.Exit"), Attributes);
+	}
+
 	// ToolsContext->EndTool only shuts the tool on the next tick, and ToolsContext->DeactivateActiveTool is
 	// inaccessible, so we end up having to do this to force the shutdown right now.
 	GetToolManager()->DeactivateTool(EToolSide::Mouse, EToolShutdownType::Cancel);
diff --git a/Engine/Plugins/Experimental/UVEditor/Source/UVEditor/Public/UVEditorMode.h b/Engine/Plugins/Experimental/UVEditor/Source/UVEditor/Public/UVEditorMode.h
--- a/Engine/Plugins/Experimental/UVEditor/Source/UVEditor/Public/UVEditorMode.h
+++ b/Engine/Plugins/Experimental/UVEditor/Source/UVEditor/Public/UVEditorMode.h
@@ -214,5 +214,7 @@ protected:
 	TArray<TObjectPtr<UInteractiveToolPropertySet>> PropertyObjectsToTick;
 
 	bool bIsActive = false;
+
+	static FDateTime AnalyticsLastStartTimestamp;
 };
 
