UVEditor: Some important bug fixes and changes to UV editor input objects and mode:
- The updating code got a lot of fixes and was largely moved into UVEditorToolUtil. It should now deal better with tricky remeshing cases, and the updates will not cause the unwrap attribute pointer to be invalidated for external holders (like a mesh change tracker).
- The reliance on mesh changestamps has been removed, since change stamps are both unreliable and likely to be soon removed entirely. Input objects now have a delegate that fires on canonical mesh updates. UV editor's stored AABB trees will update themselves from the delegate, and mode will detect changes to assets from the delegate.
- Fixed redo not working in UV editor.
- There's now a way to get camera state out of the live preview.

#rnx

diff --git a/Engine/Plugins/Experimental/UVEditor/Source/UVEditor/Private/UVEditorMode.cpp b/Engine/Plugins/Experimental/UVEditor/Source/UVEditor/Private/UVEditorMode.cpp
--- a/Engine/Plugins/Experimental/UVEditor/Source/UVEditor/Private/UVEditorMode.cpp
+++ b/Engine/Plugins/Experimental/UVEditor/Source/UVEditor/Private/UVEditorMode.cpp
@@ -400,15 +400,17 @@ void UUVEditorMode::InitializeTargets(const TArray<TObjectPtr<UObject>>& AssetsI
 		// The tool input object will hold on to the wireframe for the purposes of updating it and cleaning it up
 		ToolInputObject->WireframeDisplay = WireframeDisplay;
 
-		// Initialize our changestamp so we can later detect changes
-		LastSeenChangeStamps.Add(ToolInputObject->UnwrapCanonical->GetShapeChangeStamp());
+		// Bind to delegate so that we can detect changes
+		ToolInputObject->OnCanonicalModified.AddWeakLambda(this, [this]
+		(UUVEditorToolMeshInput* InputObject, const UUVEditorToolMeshInput::FCanonicalModifiedInfo&) {
+			ModifiedAssetIDs.Add(InputObject->AssetID);
+		});
 
 		ToolInputObjects.Add(ToolInputObject);
 	}
 
 	// Initialize our layer selector
-	UVChannelProperties->Initialize(ToolTargets, AppliedCanonicalMeshes, true);	
-	PreviousUVLayerIndex.SetNumZeroed(ToolTargets.Num());
+	UVChannelProperties->Initialize(ToolTargets, AppliedCanonicalMeshes, true);
 	PendingUVLayerIndex.SetNumZeroed(ToolTargets.Num());		
 
 	int32 AssetID = UVChannelProperties->GetSelectedAssetID();
@@ -504,43 +506,45 @@ void UUVEditorMode::ModeTick(float DeltaTime)
 
 	Super::ModeTick(DeltaTime);
 
-	bool bStartDefaultTool = false;
-	GetToolManager()->BeginUndoTransaction(UVLayerChangeTransactionName);
+	bool bSwitchingLayers = false;
 	for (int32 AssetID = 0; AssetID < ToolInputObjects.Num(); ++AssetID)
 	{
-		if ((PreviousUVLayerIndex[AssetID] != PendingUVLayerIndex[AssetID] &&
-			ToolInputObjects[AssetID]->UVLayerIndex != PendingUVLayerIndex[AssetID]) || bForceRebuildUVLayer)
+		if (ToolInputObjects[AssetID]->UVLayerIndex != PendingUVLayerIndex[AssetID] || bForceRebuildUVLayer)
 		{
-			
-
-			// TODO: Perhaps we need our own interactive tools context that allows this kind of "end tool now"
-			// call. We can't do the normal EndTool call because we cannot defer shutdown until the tick.
-			if (bPendingUVLayerChangeDestroyTool)
-			{
-				GetInteractiveToolsContext()->ToolManager->DeactivateTool(EToolSide::Mouse, EToolShutdownType::Cancel);
-			}
+			bSwitchingLayers = true;
+			break;
+		}
+	}
 
-			ChangeInputObjectLayer(AssetID, PendingUVLayerIndex[AssetID], bForceRebuildUVLayer);
+	if (bSwitchingLayers)
+	{
+		GetToolManager()->BeginUndoTransaction(UVLayerChangeTransactionName);
 
-			GetInteractiveToolsContext()->GetTransactionAPI()->AppendChange(this,
-				MakeUnique<FInputObjectUVLayerChange>(AssetID, PreviousUVLayerIndex[AssetID], PendingUVLayerIndex[AssetID]),
-				UVLayerChangeTransactionName);		
+		// TODO: Perhaps we need our own interactive tools context that allows this kind of "end tool now"
+		// call. We can't do the normal GetInteractiveToolsContext()->EndTool() call because we cannot defer
+		// shutdown.
+		GetInteractiveToolsContext()->ToolManager->DeactivateTool(EToolSide::Mouse, EToolShutdownType::Cancel);
 
-			if (bPendingUVLayerChangeDestroyTool)
+		for (int32 AssetID = 0; AssetID < ToolInputObjects.Num(); ++AssetID)
+		{
+			if (ToolInputObjects[AssetID]->UVLayerIndex != PendingUVLayerIndex[AssetID] || bForceRebuildUVLayer)
 			{
-				bStartDefaultTool = true;
+				int32 OldLayerIndex = ToolInputObjects[AssetID]->UVLayerIndex;
+
+				ChangeInputObjectLayer(AssetID, PendingUVLayerIndex[AssetID], bForceRebuildUVLayer);
+
+				GetInteractiveToolsContext()->GetTransactionAPI()->AppendChange(this,
+					MakeUnique<FInputObjectUVLayerChange>(AssetID, OldLayerIndex, PendingUVLayerIndex[AssetID]),
+					UVLayerChangeTransactionName);
 			}
-		
-			PreviousUVLayerIndex[AssetID] = PendingUVLayerIndex[AssetID];			
 		}
-	}
-	if (bStartDefaultTool)
-	{
 		ActivateDefaultTool();
-	}
-	GetToolManager()->EndUndoTransaction();
-	bForceRebuildUVLayer = false;
 
+		GetToolManager()->EndUndoTransaction();
+
+		bForceRebuildUVLayer = false;
+	}
+	
 	for (TObjectPtr<UInteractiveToolPropertySet>& Propset : PropertyObjectsToTick)
 	{
 		if (Propset)
@@ -574,13 +578,6 @@ void UUVEditorMode::ModeTick(float DeltaTime)
 		TObjectPtr<UUVEditorToolMeshInput> ToolInput = ToolInputObjects[i];
 		ToolInput->AppliedPreview->Tick(DeltaTime);
 		ToolInput->UnwrapPreview->Tick(DeltaTime);
-
-		int32 ChangeStamp = ToolInput->UnwrapCanonical->GetShapeChangeStamp();
-		if (LastSeenChangeStamps[i] != ChangeStamp)
-		{
-			LastSeenChangeStamps[i] = ChangeStamp;
-			ModifiedAssetIDs.Add(ToolInput->AssetID);
-		}
 	}
 }
 
@@ -592,7 +589,6 @@ void UUVEditorMode::SwitchActiveAsset(const FString& UVAsset)
 		// Not doing an ensure here because the "revert to default" can give us an empty string
 		UVChannelProperties->ValidateUVAssetSelection(true);
 		UpdateSelectedLayer();
-		bPendingUVLayerChangeDestroyTool = true;
 		bForceRebuildUVLayer = false;
 	}
 }
@@ -616,7 +612,6 @@ void UUVEditorMode::UpdateSelectedLayer()
 
 			UVChannelProperties->UVChannel = ChannelNames[InputObject->UVLayerIndex];
 			UVChannelProperties->SilentUpdateWatched();
-			PreviousUVLayerIndex[AssetID] = InputObject->UVLayerIndex;
 			break;
 		}
 	}
@@ -626,7 +621,7 @@ void UUVEditorMode::UpdateSelectedLayer()
 	}
 	else 
 	{
-		PendingUVLayerIndex[AssetID] = PreviousUVLayerIndex[AssetID];
+		PendingUVLayerIndex[AssetID] = ToolInputObjects[AssetID]->UVLayerIndex;
 	}		
 }
 
@@ -651,7 +646,6 @@ void UUVEditorMode::SwitchActiveChannel(const FString& UVChannel)
 			PendingUVLayerIndex[AssetID] = NewUVLayerIndex;
 		}
 	}
-	bPendingUVLayerChangeDestroyTool = true;
 	bForceRebuildUVLayer = false;
 }
 
@@ -692,7 +686,7 @@ void UUVEditorMode::ForceUpdateDisplayChannel(const TArray<int32>& LayerPerAsset
 			}
 
 			ChangeInputObjectLayer(AssetID, LayerPerAsset[AssetID], true);
-			PreviousUVLayerIndex[AssetID] = PendingUVLayerIndex[AssetID] = LayerPerAsset[AssetID];
+			PendingUVLayerIndex[AssetID] = LayerPerAsset[AssetID];
 		}
 	}
 	UpdateSelectedLayer();
diff --git a/Engine/Plugins/Experimental/UVEditor/Source/UVEditor/Private/UVEditorToolkit.cpp b/Engine/Plugins/Experimental/UVEditor/Source/UVEditor/Private/UVEditorToolkit.cpp
--- a/Engine/Plugins/Experimental/UVEditor/Source/UVEditor/Private/UVEditorToolkit.cpp
+++ b/Engine/Plugins/Experimental/UVEditor/Source/UVEditor/Private/UVEditorToolkit.cpp
@@ -11,6 +11,7 @@
 #include "SAssetEditorViewport.h"
 #include "SUVEditor2DViewport.h"
 #include "ThumbnailRendering/SceneThumbnailInfo.h"
+#include "ToolContextInterfaces.h"
 #include "ToolMenus.h"
 #include "UVEditor.h"
 #include "UVEditorMode.h"
@@ -29,6 +30,24 @@
 const FName FUVEditorToolkit::InteractiveToolsPanelTabID(TEXT("UVEditor_InteractiveToolsTab"));
 const FName FUVEditorToolkit::LivePreviewTabID(TEXT("UVEditor_LivePreviewTab"));
 
+namespace FUVEditorToolkitLocals {
+	void GetCameraState(const FEditorViewportClient& ViewportClientIn, FViewCameraState& CameraStateOut)
+	{
+		FViewportCameraTransform ViewTransform = ViewportClientIn.GetViewTransform();
+		CameraStateOut.bIsOrthographic = false;
+		CameraStateOut.bIsVR = false;
+		CameraStateOut.Position = ViewTransform.GetLocation();
+		CameraStateOut.HorizontalFOVDegrees = ViewportClientIn.ViewFOV;
+		CameraStateOut.AspectRatio = ViewportClientIn.AspectRatio;
+
+		// if using Orbit camera, the rotation in the ViewTransform is not the current camera rotation, it
+		// is set to a different rotation based on the Orbit. So we have to convert back to camera rotation.
+		FRotator ViewRotation = (ViewportClientIn.bUsingOrbitCamera) ?
+			ViewTransform.ComputeOrbitMatrix().InverseFast().Rotator() : ViewTransform.GetRotation();
+		CameraStateOut.Orientation = ViewRotation.Quaternion();
+	}
+}
+
 FUVEditorToolkit::FUVEditorToolkit(UAssetEditor* InOwningAssetEditor)
 	: FBaseAssetToolkit(InOwningAssetEditor)
 {
@@ -323,7 +342,9 @@ void FUVEditorToolkit::PostInitAssetEditor()
 	// The mode will need to be able to get to the live preview world and input.
 	UContextObjectStore* ContextStore = EditorModeManager->GetInteractiveToolsContext()->ToolManager->GetContextObjectStore();
 	UUVToolLivePreviewAPI* LivePreviewAPI = NewObject<UUVToolLivePreviewAPI>();
-	LivePreviewAPI->Initialize(LivePreviewScene->GetWorld(), LivePreviewInputRouter);
+	LivePreviewAPI->Initialize(LivePreviewScene->GetWorld(), LivePreviewInputRouter,
+		[this](FViewCameraState& CameraStateOut) {
+			FUVEditorToolkitLocals::GetCameraState(*LivePreviewViewportClient, CameraStateOut); });
 	ContextStore->AddContextObject(LivePreviewAPI);
 
 	// It is debatable where the undo/redo api should get initialized, but this seems like a 
@@ -399,7 +420,7 @@ void FUVEditorToolkit::PostInitAssetEditor()
 	{
 		ScaleFactor = UUVEditorMode::GetUVMeshScalingFactor();
 	}
-	ViewportClient->SetViewLocation(FVector(ScaleFactor / 2, ScaleFactor/2, ScaleFactor));
+	ViewportClient->SetViewLocation(FVector(ScaleFactor / 2, ScaleFactor / 2, ScaleFactor));
 	ViewportClient->SetViewRotation(FRotator(-90, 0, 0));
 
 	// If exposure isn't set to fixed, it will flash as we stare into the void
diff --git a/Engine/Plugins/Experimental/UVEditor/Source/UVEditor/Public/UVEditorMode.h b/Engine/Plugins/Experimental/UVEditor/Source/UVEditor/Public/UVEditorMode.h
--- a/Engine/Plugins/Experimental/UVEditor/Source/UVEditor/Public/UVEditorMode.h
+++ b/Engine/Plugins/Experimental/UVEditor/Source/UVEditor/Public/UVEditorMode.h
@@ -164,10 +164,6 @@ protected:
 	// Authoritative list of targets that have changes that have not been baked back yet.
 	TSet<int32> ModifiedAssetIDs;
 
-	// Helper array of canonical unwrap changestamps that allows us to detect changes to UVs.
-	// Used to populate ModifiedAssetIDs on tick.
-	TArray<int32> LastSeenChangeStamps;
-
 
 	/** Used as a selector of UV channels/layers of opened assets in the editor. */
 	UPROPERTY()
@@ -180,12 +176,9 @@ protected:
 	// Used with the ToolAssetAndLayerAPI to process tool layer change requests
 	void ForceUpdateDisplayChannel(const TArray<int32>& LayerPerAsset, bool bForceRebuildUnwrap, bool bEmitUndoTransaction);
 
-	// Used to change layers with our current picker approach (we need to remember the previous
-	// layer value so we know which one to remove). Likely to change if we start adding/removing
-	// using some different UI.
-	TArray<int32> PreviousUVLayerIndex;
 	TArray<int32> PendingUVLayerIndex;
-	bool  bPendingUVLayerChangeDestroyTool = true;
+
+	// TODO: Should be removed
 	bool  bForceRebuildUVLayer = false;
 
 
diff --git a/Engine/Plugins/Experimental/UVEditor/Source/UVEditorTools/Private/ToolTargets/UVEditorToolMeshInput.cpp b/Engine/Plugins/Experimental/UVEditor/Source/UVEditorTools/Private/ToolTargets/UVEditorToolMeshInput.cpp
--- a/Engine/Plugins/Experimental/UVEditor/Source/UVEditorTools/Private/ToolTargets/UVEditorToolMeshInput.cpp
+++ b/Engine/Plugins/Experimental/UVEditor/Source/UVEditorTools/Private/ToolTargets/UVEditorToolMeshInput.cpp
@@ -18,150 +18,11 @@ using namespace UE::Geometry;
 
 namespace UVEditorToolMeshInputLocals
 {
-	void CopyMeshPositions(const UE::Geometry::FDynamicMesh3& MeshIn, UE::Geometry::FDynamicMesh3& MeshOut,
-		const TArray<int32>* ChangedVids, const TArray<int32>* ChangedConnectivityTids)
-	{
-		auto UpdateOrInsertVid = [&MeshIn, &MeshOut](int32 Vid) {
-			if (MeshOut.IsVertex(Vid))
-			{
-				MeshOut.SetVertex(Vid, MeshIn.GetVertex(Vid));
-			}
-			else
-			{
-				MeshOut.InsertVertex(Vid, MeshIn.GetVertex(Vid));
-			}
-		};
-
-		if (!ChangedVids && !ChangedConnectivityTids)
-		{
-			for (int32 Vid : MeshIn.VertexIndicesItr())
-			{
-				UpdateOrInsertVid(Vid);
-			}
-			for (int32 Tid : MeshIn.TriangleIndicesItr())
-			{
-				MeshOut.SetTriangle(Tid, MeshIn.GetTriangle(Tid));
-			}
-
-			return;
-		}
-
-		if (ChangedVids)
-		{
-			for (int32 Vid : *ChangedVids)
-			{
-				UpdateOrInsertVid(Vid);
-			}
-		}
-
-		if (ChangedConnectivityTids)
-		{
-			for (int32 Tid : *ChangedConnectivityTids)
-			{
-				MeshOut.SetTriangle(Tid, MeshIn.GetTriangle(Tid));
-			}
-		}
-	}
-
-	/**
-	 * Copy all or parts of a mesh overlay into another mesh overlay.
-	 * 
-	 * WARNING: except in narrow cases where we are able to revert to a simple Copy() call, the function
-	 * forcibly sets the parent vertices of any elements touched by a SetTriangle() call made in the function.
-	 * This is necessary, for example, if you have a two triangle mesh with two UV islands and you want to
-	 * swap the assignment of the islands. However it means that the function could silently place the overlay 
-	 * in an invalid state if, for instance, ChangedConnectivityTids are not complete (in the previous example,
-	 * imagine only one of the triangles is included in ChangedConnectivityTids).
-	 * Consider using a checkSlow to check the validity of the output overlay to make sure you did not miss something. 
-	 *
-	 * @param bMeshesHaveSameTopology If true, underlying meshes are topologically identical, so we can use a
-	 *   simple copy when we are not constraining the copied elements/triangles. Can be set true when copying
-	 *   between a canonical/preview version of the same mesh, but must be false when copying between an unwrap
-	 *   mesh and an applied mesh.
-	 */
-	void CopyMeshOverlay(const UE::Geometry::FDynamicMeshUVOverlay& OverlayIn, UE::Geometry::FDynamicMeshUVOverlay& OverlayOut, 
-		bool bMeshesHaveSameTopology, const TArray<int32>* ChangedElements, const TArray<int32>* ChangedConnectivityTids)
-	{
-		auto UpdateOrInsertElement = [&OverlayIn, &OverlayOut](int32 ElementID) {
-			if (OverlayOut.IsElement(ElementID))
-			{
-				OverlayOut.SetElement(ElementID, OverlayIn.GetElement(ElementID));
-			}
-			else
-			{
-				FVector2f UV = OverlayIn.GetElement(ElementID);
-				OverlayOut.InsertElement(ElementID, &UV.X);
-			}
-		};
-
-		// @param PotentiallyFreedElementsOut If not null, any elements that used to be referenced by
-		// the changed triangle are placed here so they can be checked for freeing later.
-		auto ResetTriangleWithParenting = [&OverlayIn, &OverlayOut](int32 Tid, TSet<int32>* PotentiallyFreedElementsOut) {
-			if (PotentiallyFreedElementsOut)
-			{
-				FIndex3i OldElementTri = OverlayIn.GetTriangle(Tid);
-				for (int i = 0; i < 3; ++i)
-				{
-					PotentiallyFreedElementsOut->Add(OldElementTri[i]);
-				}
-			}
-
-			// Force reset the parent pointers if necessary
-			FIndex3i NewElementTri = OverlayIn.GetTriangle(Tid);
-			FIndex3i ParentTriInOutput = OverlayOut.GetParentMesh()->GetTriangle(Tid);
-			for (int32 i = 0; i < 3; ++i)
-			{
-				if (OverlayOut.GetParentVertex(NewElementTri[i]) != ParentTriInOutput[i])
-				{
-					OverlayOut.SetParentVertex(NewElementTri[i], ParentTriInOutput[i]);
-				}
-			}
-
-			// Now set the triangle. Don't free elements since other new triangles might use them.
-			OverlayOut.SetTriangle(Tid, NewElementTri, false);
-		};
-
-		if (!ChangedElements && !ChangedConnectivityTids)
-		{
-			if (bMeshesHaveSameTopology)
-			{
-				OverlayOut.Copy(OverlayIn);
-			}
-			else
-			{
-				for (int32 ElementID : OverlayIn.ElementIndicesItr())
-				{
-					UpdateOrInsertElement(ElementID);
-				}
-				for (int32 Tid : OverlayIn.GetParentMesh()->TriangleIndicesItr())
-				{
-					ResetTriangleWithParenting(Tid, nullptr);
-				}
-				OverlayOut.FreeUnusedElements();
-			}
-			return;
-		}
-
-		if (ChangedElements)
-		{
-			for (int32 ElementID : *ChangedElements)
-			{
-				UpdateOrInsertElement(ElementID);
-			}
-		}
-		if (ChangedConnectivityTids)
-		{
-			TSet<int32> PotentiallyFreedElements;
-			for (int32 Tid : *ChangedConnectivityTids)
-			{
-				ResetTriangleWithParenting(Tid, &PotentiallyFreedElements);
-			}
-
-			OverlayOut.FreeUnusedElements(&PotentiallyFreedElements);
-		}
-	}
+	const TArray<int32> EmptyArray;
 }
 
+const TArray<int32>* const UUVEditorToolMeshInput::NONE_CHANGED_ARG = &UVEditorToolMeshInputLocals::EmptyArray;
+
 bool UUVEditorToolMeshInput::IsValid() const
 {
 	return UnwrapCanonical
@@ -233,9 +94,12 @@ void UUVEditorToolMeshInput::Shutdown()
 	AppliedPreview = nullptr;
 
 	SourceTarget = nullptr;
+
+	OnCanonicalModified.Clear();
 }
 
-void UUVEditorToolMeshInput::UpdateUnwrapPreviewOverlayFromPositions(const TArray<int32>* ChangedVids, const TArray<int32>* ChangedConnectivityTids, const TArray<int32>* FastRenderUpdateTids)
+void UUVEditorToolMeshInput::UpdateUnwrapPreviewOverlayFromPositions(const TArray<int32>* ChangedVids, 
+	const TArray<int32>* ChangedConnectivityTids, const TArray<int32>* FastRenderUpdateTids)
 {
 	UnwrapPreview->PreviewMesh->DeferredEditMesh([this, ChangedVids, ChangedConnectivityTids](FDynamicMesh3& Mesh)
 	{
@@ -260,17 +124,17 @@ void UUVEditorToolMeshInput::UpdateUnwrapPreviewOverlayFromPositions(const TArra
 	}
 }
 
-void UUVEditorToolMeshInput::UpdateUnwrapCanonicalOverlayFromPositions(const TArray<int32>* ChangedVids, const TArray<int32>* ChangedConnectivityTids)
+void UUVEditorToolMeshInput::UpdateUnwrapCanonicalOverlayFromPositions(const TArray<int32>* ChangedVids, 
+	const TArray<int32>* ChangedConnectivityTids)
 {
 	FDynamicMeshUVOverlay* DestOverlay = UnwrapCanonical->Attributes()->PrimaryUV();
 	UVEditorToolUtil::UpdateUVOverlayFromUnwrapMesh(*UnwrapCanonical, *DestOverlay, VertPositionToUV,
 		ChangedVids, ChangedConnectivityTids);
 }
 
-void UUVEditorToolMeshInput::UpdateAppliedPreviewFromUnwrapPreview(const TArray<int32>* ChangedVids, const TArray<int32>* ChangedConnectivityTids, const TArray<int32>* FastRenderUpdateTids)
+void UUVEditorToolMeshInput::UpdateAppliedPreviewFromUnwrapPreview(const TArray<int32>* ChangedVids, 
+	const TArray<int32>* ChangedConnectivityTids, const TArray<int32>* FastRenderUpdateTids)
 {
-	using namespace UVEditorToolMeshInputLocals;
-
 	const FDynamicMesh3* SourceUnwrapMesh = UnwrapPreview->PreviewMesh->GetMesh();
 
 	// We need to update the overlay in AppliedPreview. Assuming that the overlay in UnwrapPreview is updated, we can
@@ -280,7 +144,7 @@ void UUVEditorToolMeshInput::UpdateAppliedPreviewFromUnwrapPreview(const TArray<
 			const FDynamicMeshUVOverlay* SourceOverlay = SourceUnwrapMesh->Attributes()->PrimaryUV();
 			FDynamicMeshUVOverlay* DestOverlay = Mesh.Attributes()->GetUVLayer(UVLayerIndex);
 
-			CopyMeshOverlay(*SourceOverlay, *DestOverlay, false, ChangedVids, ChangedConnectivityTids);
+			UVEditorToolUtil::UpdateOverlayFromOverlay(*SourceOverlay, *DestOverlay, false, ChangedVids, ChangedConnectivityTids);
 	}, false);
 
 	if (FastRenderUpdateTids) {
@@ -294,10 +158,9 @@ void UUVEditorToolMeshInput::UpdateAppliedPreviewFromUnwrapPreview(const TArray<
 	}
 }
 
-void UUVEditorToolMeshInput::UpdateUnwrapPreviewFromAppliedPreview(const TArray<int32>* ChangedElementIDs, const TArray<int32>* ChangedConnectivityTids, const TArray<int32>* FastRenderUpdateTids)
+void UUVEditorToolMeshInput::UpdateUnwrapPreviewFromAppliedPreview(const TArray<int32>* ChangedElementIDs, 
+	const TArray<int32>* ChangedConnectivityTids, const TArray<int32>* FastRenderUpdateTids)
 {
-	using namespace UVEditorToolMeshInputLocals;
-
 	// Convert the AppliedPreview UV overlay to positions in UnwrapPreview
 	const FDynamicMeshUVOverlay* SourceOverlay = AppliedPreview->PreviewMesh->GetMesh()->Attributes()->GetUVLayer(UVLayerIndex);
 	UnwrapPreview->PreviewMesh->DeferredEditMesh([this, SourceOverlay, ChangedElementIDs, ChangedConnectivityTids](FDynamicMesh3& Mesh)
@@ -306,7 +169,7 @@ void UUVEditorToolMeshInput::UpdateUnwrapPreviewFromAppliedPreview(const TArray<
 
 		// Also copy the actual overlay
 		FDynamicMeshUVOverlay* DestOverlay = Mesh.Attributes()->PrimaryUV();
-		CopyMeshOverlay(*SourceOverlay, *DestOverlay, false, ChangedElementIDs, ChangedConnectivityTids);
+		UVEditorToolUtil::UpdateOverlayFromOverlay(*SourceOverlay, *DestOverlay, false, ChangedElementIDs, ChangedConnectivityTids);
 	}, false);
 
 	if (FastRenderUpdateTids) {
@@ -325,27 +188,29 @@ void UUVEditorToolMeshInput::UpdateUnwrapPreviewFromAppliedPreview(const TArray<
 	}
 }
 
-void UUVEditorToolMeshInput::UpdateCanonicalFromPreviews(const TArray<int32>* ChangedVids, const TArray<int32>* ChangedConnectivityTids)
+void UUVEditorToolMeshInput::UpdateCanonicalFromPreviews(const TArray<int32>* ChangedVids, 
+	const TArray<int32>* ChangedConnectivityTids, bool bBroadcast)
 {
-	using namespace UVEditorToolMeshInputLocals;
-
 	// Update UnwrapCanonical from UnwrapPreview
-	UpdateOtherUnwrap(*UnwrapPreview->PreviewMesh->GetMesh(), *UnwrapCanonical, ChangedVids, ChangedConnectivityTids);
+	UVEditorToolUtil::UpdateUVUnwrapMesh(*UnwrapPreview->PreviewMesh->GetMesh(), *UnwrapCanonical, ChangedVids, ChangedConnectivityTids);
 	
 	// Update the overlay in AppliedCanonical from overlay in AppliedPreview
 	const FDynamicMeshUVOverlay* SourceOverlay = AppliedPreview->PreviewMesh->GetMesh()->Attributes()->GetUVLayer(UVLayerIndex);
 	FDynamicMeshUVOverlay* DestOverlay = AppliedCanonical->Attributes()->GetUVLayer(UVLayerIndex);
-	CopyMeshOverlay(*SourceOverlay, *DestOverlay, true, ChangedVids, ChangedConnectivityTids);
+	UVEditorToolUtil::UpdateOverlayFromOverlay(*SourceOverlay, *DestOverlay, true, ChangedVids, ChangedConnectivityTids);
+
+	if (bBroadcast)
+	{
+		OnCanonicalModified.Broadcast(this, FCanonicalModifiedInfo());
+	}
 }
 
 void UUVEditorToolMeshInput::UpdatePreviewsFromCanonical(const TArray<int32>* ChangedVids, const TArray<int32>* ChangedConnectivityTids, const TArray<int32>* FastRenderUpdateTids)
 {
-	using namespace UVEditorToolMeshInputLocals;
-
 	// Update UnwrapPreview from UnwrapCanonical
 	UnwrapPreview->PreviewMesh->DeferredEditMesh([this, ChangedVids, ChangedConnectivityTids](FDynamicMesh3& Mesh)
 	{
-		UpdateOtherUnwrap(*UnwrapCanonical, Mesh, ChangedVids, ChangedConnectivityTids);
+		UVEditorToolUtil::UpdateUVUnwrapMesh(*UnwrapCanonical, Mesh, ChangedVids, ChangedConnectivityTids);
 	}, false);
 	if (FastRenderUpdateTids) {
 		UnwrapPreview->PreviewMesh->NotifyRegionDeferredEditCompleted(*FastRenderUpdateTids,
@@ -367,7 +232,7 @@ void UUVEditorToolMeshInput::UpdatePreviewsFromCanonical(const TArray<int32>* Ch
 	{
 		const FDynamicMeshUVOverlay* SourceOverlay = AppliedCanonical->Attributes()->GetUVLayer(UVLayerIndex);
 		FDynamicMeshUVOverlay* DestOverlay = Mesh.Attributes()->GetUVLayer(UVLayerIndex);
-		CopyMeshOverlay(*SourceOverlay, *DestOverlay, true, ChangedVids, ChangedConnectivityTids);
+		UVEditorToolUtil::UpdateOverlayFromOverlay(*SourceOverlay, *DestOverlay, true, ChangedVids, ChangedConnectivityTids);
 	}, false);
 	if (FastRenderUpdateTids) {
 		AppliedPreview->PreviewMesh->NotifyRegionDeferredEditCompleted(*FastRenderUpdateTids,
@@ -380,66 +245,66 @@ void UUVEditorToolMeshInput::UpdatePreviewsFromCanonical(const TArray<int32>* Ch
 	}
 }
 
-void UUVEditorToolMeshInput::UpdateAllFromUnwrapPreview(const TArray<int32>* ChangedVids, const TArray<int32>* ChangedConnectivityTids, const TArray<int32>* FastRenderUpdateTids)
+void UUVEditorToolMeshInput::UpdateAllFromUnwrapPreview(const TArray<int32>* ChangedVids, 
+	const TArray<int32>* ChangedConnectivityTids, const TArray<int32>* FastRenderUpdateTids, bool bBroadcast)
 {
 	UpdateAppliedPreviewFromUnwrapPreview(ChangedVids, ChangedConnectivityTids, FastRenderUpdateTids);
-	UpdateCanonicalFromPreviews(ChangedVids, ChangedConnectivityTids);
+	UpdateCanonicalFromPreviews(ChangedVids, ChangedConnectivityTids, false);
+
+	if (bBroadcast)
+	{
+		OnCanonicalModified.Broadcast(this, FCanonicalModifiedInfo());
+	}
 }
 
 void UUVEditorToolMeshInput::UpdateAllFromUnwrapCanonical(
-	const TArray<int32>* ChangedVids, const TArray<int32>* ChangedConnectivityTids, const TArray<int32>* FastRenderUpdateTids)
+	const TArray<int32>* ChangedVids, const TArray<int32>* ChangedConnectivityTids, 
+	const TArray<int32>* FastRenderUpdateTids, bool bBroadcast)
 {
-	using namespace UVEditorToolMeshInputLocals;
-
 	// Update AppliedCanonical
 	FDynamicMeshUVOverlay* SourceOverlay = UnwrapCanonical->Attributes()->PrimaryUV();
 	FDynamicMeshUVOverlay* DestOverlay = AppliedCanonical->Attributes()->GetUVLayer(UVLayerIndex);
-	CopyMeshOverlay(*SourceOverlay, *DestOverlay, false, ChangedVids, ChangedConnectivityTids);
+	UVEditorToolUtil::UpdateOverlayFromOverlay(*SourceOverlay, *DestOverlay, false, ChangedVids, ChangedConnectivityTids);
 
 	UpdatePreviewsFromCanonical(ChangedVids, ChangedConnectivityTids, FastRenderUpdateTids);
+
+	if (bBroadcast)
+	{
+		OnCanonicalModified.Broadcast(this, FCanonicalModifiedInfo());
+	}
 }
 
 void UUVEditorToolMeshInput::UpdateAllFromAppliedCanonical(
-	const TArray<int32>* ChangedElementIDs, const TArray<int32>* ChangedConnectivityTids, const TArray<int32>* FastRenderUpdateTids)
+	const TArray<int32>* ChangedElementIDs, const TArray<int32>* ChangedConnectivityTids, 
+	const TArray<int32>* FastRenderUpdateTids, bool bBroadcast)
 {
-	using namespace UVEditorToolMeshInputLocals;
-
 	// Update UnwrapCanonical
 	const FDynamicMeshUVOverlay* SourceOverlay = AppliedCanonical->Attributes()->GetUVLayer(UVLayerIndex);
 	UVEditorToolUtil::UpdateUVUnwrapMesh(*SourceOverlay, *UnwrapCanonical, UVToVertPosition, ChangedElementIDs, ChangedConnectivityTids);
 
 	UpdatePreviewsFromCanonical(ChangedElementIDs, ChangedConnectivityTids, FastRenderUpdateTids);
+
+	if (bBroadcast)
+	{
+		OnCanonicalModified.Broadcast(this, FCanonicalModifiedInfo());
+	}
 }
 
 void UUVEditorToolMeshInput::UpdateAllFromAppliedPreview(
-	const TArray<int32>* ChangedElementIDs, const TArray<int32>* ChangedConnectivityTids, const TArray<int32>* FastRenderUpdateTids)
+	const TArray<int32>* ChangedElementIDs, const TArray<int32>* ChangedConnectivityTids, 
+	const TArray<int32>* FastRenderUpdateTids, bool bBroadcast)
 {
 	UpdateUnwrapPreviewFromAppliedPreview(ChangedElementIDs, ChangedConnectivityTids, FastRenderUpdateTids);
-	UpdateCanonicalFromPreviews(ChangedElementIDs, ChangedConnectivityTids);
-}
+	UpdateCanonicalFromPreviews(ChangedElementIDs, ChangedConnectivityTids, false);
 
-/**
- * Helper function. Uses the positions and UV overlay of one unwrap mesh to update another one.
- */
-void UUVEditorToolMeshInput::UpdateOtherUnwrap(const FDynamicMesh3& SourceUnwrapMesh,
-	FDynamicMesh3& DestUnwrapMesh, const TArray<int32>* ChangedVids, const TArray<int32>* ChangedConnectivityTids)
-{
-	using namespace UVEditorToolMeshInputLocals;
-
-	if (!ChangedVids && !ChangedConnectivityTids)
+	if (bBroadcast)
 	{
-		DestUnwrapMesh.Copy(SourceUnwrapMesh, false, false, false, true); // Copy positions and UVs
-		return;
+		OnCanonicalModified.Broadcast(this, FCanonicalModifiedInfo());
 	}
-
-	CopyMeshPositions(SourceUnwrapMesh, DestUnwrapMesh, ChangedVids, ChangedConnectivityTids);
-
-	const FDynamicMeshUVOverlay* SourceOverlay = SourceUnwrapMesh.Attributes()->PrimaryUV();
-	FDynamicMeshUVOverlay* DestOverlay = DestUnwrapMesh.Attributes()->PrimaryUV();
-	CopyMeshOverlay(*SourceOverlay, *DestOverlay, true, ChangedVids, ChangedConnectivityTids);
 }
 
-void UUVEditorToolMeshInput::UpdateFromCanonicalUnwrapUsingMeshChange(const FDynamicMeshChange& UnwrapCanonicalMeshChange)
+void UUVEditorToolMeshInput::UpdateFromCanonicalUnwrapUsingMeshChange(
+	const FDynamicMeshChange& UnwrapCanonicalMeshChange, bool bBroadcast)
 {
 	// Note that we know that no triangles were created or destroyed since the UV editor
 	// does not allow that (it would break the mesh mappings). Otherwise we would need to
@@ -454,9 +319,22 @@ void UUVEditorToolMeshInput::UpdateFromCanonicalUnwrapUsingMeshChange(const FDyn
 	VertexToTriangleOneRing(UnwrapCanonical.Get(), ChangedVids, RenderUpdateTidsSet);
 	TArray<int32> RenderUpdateTids = RenderUpdateTidsSet.Array();
 
-	// Todo: Determine a better solution for supporting explicit changed Tids and Vids
-	// 	     from MeshChanges. This current approach only works well if the topology is not
-	// 	     changing between updates.
-	//UpdateAllFromUnwrapCanonical(&ChangedVids, & ChangedTids, & RenderUpdateTids);
-	UpdateAllFromUnwrapCanonical(nullptr, nullptr, nullptr);
+	UpdateAllFromUnwrapCanonical(&ChangedVids, &ChangedTids, &RenderUpdateTids, false);
+	
+	if (bBroadcast)
+	{
+		OnCanonicalModified.Broadcast(this, FCanonicalModifiedInfo());
+	}
 }
+
+int32 UUVEditorToolMeshInput::UnwrapVidToAppliedVid(int32 UnwrapVid)
+{
+	const FDynamicMeshUVOverlay* CanonicalOverlay = AppliedCanonical->Attributes()->GetUVLayer(UVLayerIndex);
+	return CanonicalOverlay->GetParentVertex(UnwrapVid);
+}
+
+void UUVEditorToolMeshInput::AppliedVidToUnwrapVids(int32 AppliedVid, TArray<int32>& UnwrapVidsOut)
+{
+	const FDynamicMeshUVOverlay* CanonicalOverlay = AppliedCanonical->Attributes()->GetUVLayer(UVLayerIndex);
+	CanonicalOverlay->GetVertexElements(AppliedVid, UnwrapVidsOut);
+}
\ No newline at end of file
diff --git a/Engine/Plugins/Experimental/UVEditor/Source/UVEditorTools/Private/UVEditorLayerEditTool.cpp b/Engine/Plugins/Experimental/UVEditor/Source/UVEditorTools/Private/UVEditorLayerEditTool.cpp
--- a/Engine/Plugins/Experimental/UVEditor/Source/UVEditorTools/Private/UVEditorLayerEditTool.cpp
+++ b/Engine/Plugins/Experimental/UVEditor/Source/UVEditorTools/Private/UVEditorLayerEditTool.cpp
@@ -49,7 +49,7 @@ namespace UUVEditorChannelEditLocals
 				Target->AppliedPreview->PreviewMesh->UpdatePreview(Target->AppliedCanonical.Get());
 
 				AssetAndChannelAPI->NotifyOfAssetChannelCountChange(Target->AssetID);
-				Target->OnUndoRedo.Broadcast(false);
+				Target->OnCanonicalModified.Broadcast(Target, UUVEditorToolMeshInput::FCanonicalModifiedInfo());
 			}
 
 		}
@@ -63,7 +63,7 @@ namespace UUVEditorChannelEditLocals
 			int32 NewChannelIndex = DynamicMeshUVEditor.RemoveUVLayer();
 			Target->AppliedPreview->PreviewMesh->UpdatePreview(Target->AppliedCanonical.Get());
 			AssetAndChannelAPI->NotifyOfAssetChannelCountChange(Target->AssetID);
-			Target->OnUndoRedo.Broadcast(true);		
+			Target->OnCanonicalModified.Broadcast(Target, UUVEditorToolMeshInput::FCanonicalModifiedInfo());
 		}
 
 		virtual bool HasExpired(UObject* Object) const override
@@ -100,7 +100,6 @@ namespace UUVEditorChannelEditLocals
 			FDynamicMeshUVEditor DynamicMeshUVEditor(Target->AppliedCanonical.Get(), TargetUVChannelIndex, false);
 			DynamicMeshUVEditor.CopyUVLayer(Target->AppliedCanonical->Attributes()->GetUVLayer(SourceUVChannelIndex));
 			Target->UpdateAllFromAppliedCanonical();
-			Target->OnUndoRedo.Broadcast(false);
 		}
 
 		virtual void Revert(UObject* Object) override
@@ -111,7 +110,6 @@ namespace UUVEditorChannelEditLocals
 			FDynamicMeshUVEditor DynamicMeshUVEditor(Target->AppliedCanonical.Get(), TargetUVChannelIndex, false);
 			DynamicMeshUVEditor.CopyUVLayer(&OriginalUVChannel);
 			Target->UpdateAllFromAppliedCanonical();
-			Target->OnUndoRedo.Broadcast(true);
 		}
 
 		virtual bool HasExpired(UObject* Object) const override
@@ -158,8 +156,8 @@ namespace UUVEditorChannelEditLocals
 			TArray<int32> ChannelPerAsset = AssetAndChannelAPI->GetCurrentChannelVisibility();
 			ChannelPerAsset[Target->AssetID] = DeletedUVChannelIndex;
 			AssetAndChannelAPI->RequestChannelVisibilityChange(ChannelPerAsset, false, false);
-
-			Target->OnUndoRedo.Broadcast(false);
+			
+			Target->OnCanonicalModified.Broadcast(Target, UUVEditorToolMeshInput::FCanonicalModifiedInfo());
 		}
 
 		virtual void Revert(UObject* Object) override
@@ -183,8 +181,8 @@ namespace UUVEditorChannelEditLocals
 			TArray<int32> ChannelPerAsset = AssetAndChannelAPI->GetCurrentChannelVisibility();
 			ChannelPerAsset[Target->AssetID] = DeletedUVChannelIndex;
 			AssetAndChannelAPI->RequestChannelVisibilityChange(ChannelPerAsset, true, false);
-
-			Target->OnUndoRedo.Broadcast(true);
+			
+			Target->OnCanonicalModified.Broadcast(Target, UUVEditorToolMeshInput::FCanonicalModifiedInfo());
 		}
 
 		virtual bool HasExpired(UObject* Object) const override
@@ -457,9 +455,10 @@ void UUVEditorChannelEditTool::Setup()
 
 	for (int32 i = 0; i < Targets.Num(); ++i)
 	{
-		Targets[i]->OnUndoRedo.AddWeakLambda(this, [this, i](bool bRevert) {
+		Targets[i]->OnCanonicalModified.AddWeakLambda(this, [this, i]
+		(UUVEditorToolMeshInput* InputObject, const UUVEditorToolMeshInput::FCanonicalModifiedInfo&) {
 			UpdateChannelSelectionProperties(i);
-			});
+		});
 	}
 }
 
@@ -467,7 +466,7 @@ void UUVEditorChannelEditTool::Shutdown(EToolShutdownType ShutdownType)
 {
 	for (TObjectPtr<UUVEditorToolMeshInput> Target : Targets)
 	{
-		Target->OnUndoRedo.RemoveAll(this);
+		Target->OnCanonicalModified.RemoveAll(this);
 	}
 
 	Targets.Empty();
diff --git a/Engine/Plugins/Experimental/UVEditor/Source/UVEditorTools/Private/UVEditorToolUtil.cpp b/Engine/Plugins/Experimental/UVEditor/Source/UVEditorTools/Private/UVEditorToolUtil.cpp
--- a/Engine/Plugins/Experimental/UVEditor/Source/UVEditorTools/Private/UVEditorToolUtil.cpp
+++ b/Engine/Plugins/Experimental/UVEditor/Source/UVEditorTools/Private/UVEditorToolUtil.cpp
@@ -7,6 +7,80 @@
 
 using namespace UE::Geometry;
 
+namespace UVEditorToolUtilLocals
+{
+	// Helper function to update the triangles in an unwrap mesh based on an overlay.
+	template <typename TriIteratorType>
+	void UpdateUnwrapTriangles(const FDynamicMeshUVOverlay& TriSource,
+		const TriIteratorType& TriIterator, FDynamicMesh3& UnwrapMeshOut)
+	{
+		// Updating triangles is a little messy. To be able to handle arbitrary remeshing, we
+		// have to delete all tris first before adding any, so that we don't fail in SetTriangle
+		// due to (temporary) non-manifold edges. We have to do this without removing verts, which
+		// means we later need to check for isolated verts. Finally, because we don't have a way
+		// to avoid removing temporarily isolated UV elements, we need to reattach the UV elements
+		// for verts that were temporarily left isolated.
+
+		FDynamicMeshUVOverlay* UnwrapMeshUVOverlay = UnwrapMeshOut.Attributes()->PrimaryUV();
+
+		// Remove tris and keep track of potentially isolated elements
+		TSet<int32> PotentiallyIsolatedElements;
+		for (int32 Tid : TriIterator)
+		{
+			// Shouldn't have to have this check, but it guards against TriIterator having duplicates.
+			if (!UnwrapMeshOut.IsTriangle(Tid))
+			{
+				// We probably didn't mean to be looking at the same tids multiple times, so something probably
+				// isn't quite right upstream.
+				ensure(false);
+				continue;
+			}
+
+			FIndex3i PrevTriangle = UnwrapMeshOut.GetTriangle(Tid);
+			for (int i = 0; i < 3; ++i)
+			{
+				PotentiallyIsolatedElements.Add(PrevTriangle[i]);
+			}
+			UnwrapMeshOut.RemoveTriangle(Tid, false);
+		}
+
+		// Reinsert new tris
+		for (int32 Tid : TriIterator)
+		{
+			// Shouldn't have to have this check, but it guards against TriIterator having duplicates.
+			if (UnwrapMeshOut.IsTriangle(Tid))
+			{
+				continue;
+			}
+
+			FIndex3i NewTriangle = TriSource.GetTriangle(Tid);
+			UnwrapMeshOut.InsertTriangle(Tid, NewTriangle);
+		}
+
+		// Deal with isolated and non-isolated verts
+		for (int32 ElementID : PotentiallyIsolatedElements)
+		{
+			if (!UnwrapMeshOut.IsReferencedVertex(ElementID))
+			{
+				UnwrapMeshOut.RemoveVertex(ElementID);
+			}
+			else if (!UnwrapMeshUVOverlay->IsElement(ElementID))
+			{
+				// This is a referenced vert without a UV element (because it got removed
+				// during temporary isolation), so reinstate the element.
+				FVector2f ElementValue = TriSource.GetElement(ElementID);
+				UnwrapMeshUVOverlay->InsertElement(ElementID, &ElementValue.X);
+			}
+		}
+
+		// Update overlay tris now that we know that the elements exist.
+		for (int32 Tid : TriIterator)
+		{
+			UnwrapMeshUVOverlay->SetTriangle(Tid, TriSource.GetTriangle(Tid));
+		}
+	}
+}
+
 void UVEditorToolUtil::GenerateUVUnwrapMesh(
 	const FDynamicMeshUVOverlay& UVOverlay, FDynamicMesh3& UnwrapMeshOut,
 	TFunctionRef<FVector3d(const FVector2f&)> UVToVertPosition)
@@ -50,20 +124,31 @@ void UVEditorToolUtil::UpdateUVUnwrapMesh(const FDynamicMeshUVOverlay& UVOverlay
 	FDynamicMesh3& UnwrapMeshOut, TFunctionRef<FVector3d(const FVector2f&)> UVToVertPosition,
 	const TArray<int32>* ChangedElementIDs, const TArray<int32>* ChangedTids)
 {
-	if (!ChangedElementIDs && !ChangedTids)
-	{
-		UVEditorToolUtil::GenerateUVUnwrapMesh(UVOverlayIn, UnwrapMeshOut, UVToVertPosition);
-		return;
-	}
+	using namespace UVEditorToolUtilLocals;
 
-	// Otherwise, we only update the parts that we're asked.
-	FDynamicMeshUVOverlay* UnwrapMeshUVOverlay = UnwrapMeshOut.Attributes()->GetUVLayer(0);
-	if (ChangedElementIDs)
+	// Note that we don't want to use GenerateUVUnwrapMesh even when doing a full update
+	// because that clears the mesh and rebuilds it, and that resets the attributes pointer.
+	// That would prevent us from using a dynamic mesh change tracker across an update, as
+	// it would lose its attribute pointer.
+
+	FDynamicMeshUVOverlay* UnwrapMeshUVOverlay = UnwrapMeshOut.Attributes()->PrimaryUV();
+
+	auto UpdateVertPositions = [&UVOverlayIn, &UnwrapMeshOut, UVToVertPosition, UnwrapMeshUVOverlay](const auto& ElementIterator)
 	{
-		for (int32 ElementID : *ChangedElementIDs)
+		for (int32 ElementID : ElementIterator)
 		{
+			if (!ensure(UVOverlayIn.IsElement(ElementID)))
+			{
+				// [ELEMENT_NOT_IN_SOURCE]
+				// If you ended up here, then you asked to update an element that wasn't in the source mesh.
+				// Perhaps you gathered the changing elements pre-change, and that element was deleted. You 
+				// shouldn't gather pre-change because you risk not including any added elements, and because 
+				// deleted elements should be captured by changed tri connectivity.
+				continue;
+			}
+
 			FVector2f ElementValue = UVOverlayIn.GetElement(ElementID);
-			
+
 			// Update the actual unwrap mesh
 			if (UnwrapMeshOut.IsVertex(ElementID))
 			{
@@ -74,7 +159,7 @@ void UVEditorToolUtil::UpdateUVUnwrapMesh(const FDynamicMeshUVOverlay& UVOverlay
 				UnwrapMeshOut.InsertVertex(ElementID, UVToVertPosition(ElementValue));
 			}
 
-			// Update the unwrap overlay
+			// Update the unwrap overlay.
 			if (UnwrapMeshUVOverlay->IsElement(ElementID))
 			{
 				UnwrapMeshUVOverlay->SetElement(ElementID, ElementValue);
@@ -84,16 +169,80 @@ void UVEditorToolUtil::UpdateUVUnwrapMesh(const FDynamicMeshUVOverlay& UVOverlay
 				UnwrapMeshUVOverlay->InsertElement(ElementID, &ElementValue.X);
 			}
 		}
+	};
+
+	if (ChangedElementIDs)
+	{
+		UpdateVertPositions(*ChangedElementIDs);
+	}
+	else
+	{
+		UpdateVertPositions(UVOverlayIn.ElementIndicesItr());
 	}
 
 	if (ChangedTids)
 	{
-		for (int32 Tid : *ChangedTids)
+		UpdateUnwrapTriangles(UVOverlayIn, *ChangedTids, UnwrapMeshOut);
+	}
+	else
+	{
+		UpdateUnwrapTriangles(UVOverlayIn, UVOverlayIn.GetParentMesh()->TriangleIndicesItr(), UnwrapMeshOut);
+	}
+}
+
+void UVEditorToolUtil::UpdateUVUnwrapMesh(const FDynamicMesh3& SourceUnwrapMesh, FDynamicMesh3& DestUnwrapMesh,
+	const TArray<int32>* ChangedVids, const TArray<int32>* ChangedConnectivityTids)
+{
+	using namespace UVEditorToolUtilLocals;
+
+	if (!ChangedVids && !ChangedConnectivityTids)
+	{
+		DestUnwrapMesh.Copy(SourceUnwrapMesh, false, false, false, true); // Copy positions and UVs
+		return;
+	}
+
+	const FDynamicMeshUVOverlay* SourceOverlay = SourceUnwrapMesh.Attributes()->PrimaryUV();
+	FDynamicMeshUVOverlay* DestOverlay = DestUnwrapMesh.Attributes()->PrimaryUV();
+
+	auto UpdateVerts = [&SourceUnwrapMesh, &DestUnwrapMesh, SourceOverlay, DestOverlay](const auto& VidIterator)
+	{
+		for (int32 Vid : VidIterator)
 		{
-			FIndex3i Triangle = UVOverlayIn.GetTriangle(Tid);
-			UnwrapMeshUVOverlay->SetTriangle(Tid, Triangle);
-			UnwrapMeshOut.SetTriangle(Tid, Triangle);
+			if (!ensure(SourceUnwrapMesh.IsVertex(Vid)))
+			{
+				// See comment labeled [ELEMENT_NOT_IN_SOURCE] above.
+				continue;
+			}
+
+			if (DestUnwrapMesh.IsVertex(Vid))
+			{
+				DestUnwrapMesh.SetVertex(Vid, SourceUnwrapMesh.GetVertex(Vid));
+				DestOverlay->SetElement(Vid, SourceOverlay->GetElement(Vid));
+			}
+			else
+			{
+				DestUnwrapMesh.InsertVertex(Vid, SourceUnwrapMesh.GetVertex(Vid));
+				FVector2f ElementValue = SourceOverlay->GetElement(Vid);
+				DestOverlay->InsertElement(Vid, &ElementValue.X);
+			}
 		}
+	};
+	if (ChangedVids)
+	{
+		UpdateVerts(*ChangedVids);
+	}
+	else
+	{
+		UpdateVerts(SourceUnwrapMesh.VertexIndicesItr());
+	}
+
+	if (ChangedConnectivityTids)
+	{
+		UpdateUnwrapTriangles(*SourceOverlay, *ChangedConnectivityTids, DestUnwrapMesh);
+	}
+	else
+	{
+		UpdateUnwrapTriangles(*SourceOverlay, SourceUnwrapMesh.TriangleIndicesItr(), DestUnwrapMesh);
 	}
 }
 
@@ -107,46 +256,155 @@ void UVEditorToolUtil::UpdateUVOverlayFromUnwrapMesh(
 		return;
 	}
 
-	auto UpdateUVElement = [&UnwrapMeshIn, &UVOverlayOut, VertPositionToUV](int32 Vid) {
-		if (UVOverlayOut.IsElement(Vid))
-		{
-			UVOverlayOut.SetElement(Vid, VertPositionToUV(UnwrapMeshIn.GetVertex(Vid)));
-		}
-		else
+	auto UpdateElements = [&UnwrapMeshIn, &UVOverlayOut, VertPositionToUV](const auto& VidIterator)
+	{
+		for (int32 Vid : VidIterator)
 		{
+			if (!ensure(UnwrapMeshIn.IsVertex(Vid)))
+			{
+				// See comment labeled [ELEMENT_NOT_IN_SOURCE] above.
+				continue;
+			}
+
 			FVector2f UV = VertPositionToUV(UnwrapMeshIn.GetVertex(Vid));
-			UVOverlayOut.InsertElement(Vid, &UV.X);
+
+			if (UVOverlayOut.IsElement(Vid))
+			{
+				UVOverlayOut.SetElement(Vid, UV);
+			}
+			else
+			{
+				UVOverlayOut.InsertElement(Vid, &UV.X);
+			}
 		}
 	};
+	if (ChangedVids)
+	{
+		UpdateElements(*ChangedVids);
+	}
+	else
+	{
+		UpdateElements(UnwrapMeshIn.VertexIndicesItr());
+	}
 
-	// Update all vids and tids if not given changed items
-	if (!ChangedVids && !ChangedTids)
+	auto UpdateTriangles = [&UnwrapMeshIn, &UVOverlayOut](const auto& TriIterator)
 	{
-		for (int32 Vid : UnwrapMeshIn.VertexIndicesItr())
+		TSet<int32> PotentiallyFreedElements;
+		for (int32 Tid : TriIterator)
 		{
-			UpdateUVElement(Vid);
-		}
-		for (int32 Tid : UnwrapMeshIn.TriangleIndicesItr())
-		{
-			UVOverlayOut.SetTriangle(Tid, UnwrapMeshIn.GetTriangle(Tid));
+			FIndex3i NewTriangle = UnwrapMeshIn.GetTriangle(Tid);
+			FIndex3i PrevTriangle = UVOverlayOut.GetTriangle(Tid);
+			for (int i = 0; i < 3; ++i)
+			{
+				if (NewTriangle[i] != PrevTriangle[i])
+				{
+					PotentiallyFreedElements.Add(PrevTriangle[i]);
+				}
+			}
+			UVOverlayOut.SetTriangle(Tid, NewTriangle, false);
 		}
 
+		UVOverlayOut.FreeUnusedElements(&PotentiallyFreedElements);
+	};
+	if (ChangedTids)
+	{
+		UpdateTriangles(*ChangedTids);
+	}
+	else
+	{
+		UpdateTriangles(UnwrapMeshIn.TriangleIndicesItr());
+	}
+}
+
+void UVEditorToolUtil::UpdateOverlayFromOverlay(
+	const FDynamicMeshUVOverlay& OverlayIn, FDynamicMeshUVOverlay& OverlayOut,
+	bool bMeshesHaveSameTopology, const TArray<int32>* ChangedElements, 
+	const TArray<int32>* ChangedConnectivityTids)
+{
+	if (!ChangedElements && !ChangedConnectivityTids && bMeshesHaveSameTopology)
+	{
+		OverlayOut.Copy(OverlayIn);
 		return;
 	}
 
-	//Otherwise, only update what we're told about.
-	if (ChangedVids)
+	auto UpdateElements = [&OverlayIn, &OverlayOut](const auto& ElementIterator)
 	{
-		for (int32 Vid : *ChangedVids)
+		for (int32 ElementID : ElementIterator)
 		{
-			UpdateUVElement(Vid);
+			if (!ensure(OverlayIn.IsElement(ElementID)))
+			{
+				// See comment labeled [ELEMENT_NOT_IN_SOURCE] above.
+				continue;
+			}
+
+			FVector2f ElementValue = OverlayIn.GetElement(ElementID);
+			if (OverlayOut.IsElement(ElementID))
+			{
+				OverlayOut.SetElement(ElementID, ElementValue);
+			}
+			else
+			{
+				OverlayOut.InsertElement(ElementID, &ElementValue.X);
+			}
 		}
+	};
+
+	if (ChangedElements)
+	{
+		UpdateElements(*ChangedElements);
 	}
-	if (ChangedTids)
+	else
 	{
-		for (int32 Tid : *ChangedTids)
+		UpdateElements(OverlayIn.ElementIndicesItr());
+	}
+
+	auto UpdateTriangles = [&OverlayIn, &OverlayOut](const auto& TriIterator)
+	{
+		// To handle arbitrary remeshing in the UV overlay, not only do we need to
+		// check for freed elements only after finishing the updates, but we may 
+		// also need to forcefully change the parent pointer of elements (imagine
+		// a mesh of two disconnected triangles whose element mappings changed)
+		TSet<int32> PotentiallyFreedElements;
+		for (int32 Tid : TriIterator)
 		{
-			UVOverlayOut.SetTriangle(Tid, UnwrapMeshIn.GetTriangle(Tid));
+			FIndex3i OldElementTri = OverlayIn.GetTriangle(Tid);
+			FIndex3i NewElementTri = OverlayIn.GetTriangle(Tid);
+			FIndex3i ParentTriInOutput = OverlayOut.GetParentMesh()->GetTriangle(Tid);
+
+			for (int i = 0; i < 3; ++i)
+			{
+				if (NewElementTri[i] != OldElementTri[i])
+				{
+					PotentiallyFreedElements.Add(OldElementTri[i]);
+				}
+
+				// Force the parent pointer, if needed
+				if (OverlayOut.GetParentVertex(NewElementTri[i]) != ParentTriInOutput[i])
+				{
+					OverlayOut.SetParentVertex(NewElementTri[i], ParentTriInOutput[i]);
+				}
+			}
+
+			for (int32 i = 0; i < 3; ++i)
+			{
+				if (OverlayOut.GetParentVertex(NewElementTri[i]) != ParentTriInOutput[i])
+				{
+					OverlayOut.SetParentVertex(NewElementTri[i], ParentTriInOutput[i]);
+				}
+			}
+
+			OverlayOut.SetTriangle(Tid, NewElementTri, false);
 		}
+
+		OverlayOut.FreeUnusedElements(&PotentiallyFreedElements);
+	};
+
+	if (ChangedConnectivityTids)
+	{
+		UpdateTriangles(*ChangedConnectivityTids);
+	}
+	else
+	{
+		UpdateTriangles(OverlayIn.GetParentMesh()->TriangleIndicesItr());
 	}
 }
diff --git a/Engine/Plugins/Experimental/UVEditor/Source/UVEditorTools/Private/UVSelectTool.cpp b/Engine/Plugins/Experimental/UVEditor/Source/UVEditorTools/Private/UVSelectTool.cpp
--- a/Engine/Plugins/Experimental/UVEditor/Source/UVEditorTools/Private/UVSelectTool.cpp
+++ b/Engine/Plugins/Experimental/UVEditor/Source/UVEditorTools/Private/UVSelectTool.cpp
@@ -5,7 +5,6 @@
 #include "Algo/Unique.h"
 #include "BaseGizmos/CombinedTransformGizmo.h"
 #include "ContextObjectStore.h"
-#include "DynamicMesh/DynamicMeshChangeTracker.h"
 #include "Drawing/LineSetComponent.h"
 #include "Drawing/MeshElementsVisualizer.h"
 #include "Drawing/PreviewGeometryActor.h"
@@ -14,6 +13,7 @@
 #include "ToolTargets/UVEditorToolMeshInput.h"
 #include "InteractiveToolManager.h"
 #include "MeshOpPreviewHelpers.h" // UMeshOpPreviewWithBackgroundCompute
+#include "Parameterization/DynamicMeshUVEditor.h"
 #include "PreviewMesh.h"
 #include "Selection/MeshSelectionMechanic.h"
 #include "Selection/DynamicMeshSelection.h"
@@ -133,10 +133,6 @@ namespace UVSelectToolLocals
 		{
 			UnwrapCanonicalMeshChange->Apply(UVToolInputObject->UnwrapCanonical.Get(), false);
 			UVToolInputObject->UpdateFromCanonicalUnwrapUsingMeshChange(*UnwrapCanonicalMeshChange);
-			
-			// This is a little wasteful because we're going to reset the gizmo transform, 
-			// but it updates the AABBTree for us.
-			UVToolInputObject->OnUndoRedo.Broadcast(false);
 
 			UUVSelectToolChangeRouter* ChangeRouter = Cast<UUVSelectToolChangeRouter>(Object);
 			if (ensure(ChangeRouter) && ChangeRouter->CurrentSelectTool.IsValid())
@@ -150,10 +146,6 @@ namespace UVSelectToolLocals
 		{
 			UnwrapCanonicalMeshChange->Apply(UVToolInputObject->UnwrapCanonical.Get(), true);
 			UVToolInputObject->UpdateFromCanonicalUnwrapUsingMeshChange(*UnwrapCanonicalMeshChange);
-			
-			// This is a little wasteful because we're going to reset the gizmo transform, 
-			// but it updates the AABBTree for us.
-			UVToolInputObject->OnUndoRedo.Broadcast(false);
 
 			UUVSelectToolChangeRouter* ChangeRouter = Cast<UUVSelectToolChangeRouter>(Object);
 			if (ensure(ChangeRouter) && ChangeRouter->CurrentSelectTool.IsValid())
@@ -211,7 +203,7 @@ UInteractiveTool* UUVSelectToolBuilder::BuildTool(const FToolBuilderState& Scene
 
 void  USelectToolActionPropertySet::Sew()
 {
-	PostAction(ESelectToolAction::Sew); 
+	PostAction(ESelectToolAction::Sew);
 }
 
 void USelectToolActionPropertySet::PostAction(ESelectToolAction Action)
@@ -289,18 +281,7 @@ void UUVSelectTool::Setup()
 		{
 			Tree = MakeShared<FDynamicMeshAABBTree3>();
 			Tree->SetMesh(Target->UnwrapCanonical.Get());
-			TreeStore->Set(Target->UnwrapCanonical.Get(), Tree);
-		}
-		// TODO: We currently can't do this check because the aabb tree checks the mesh
-		// changestamp, and that gets (arguably incorrectly) reset if we do an update
-		// that involves a clear. So, for instance, after a layout operation, the changestamp
-		// frequently ends up matching the original and incorrectly passes the check.
-		// All this means that storing the tree is currently pointless because
-		// we rebuild it each time we switch back to this tool, but we keep the code
-		// until we fix the change stamp thing or give a workaround.
-		//if (!Tree->IsValid(false))
-		{
-			Tree->Build();
+			TreeStore->Set(Target->UnwrapCanonical.Get(), Tree, Target);
 		}
 		AABBTrees.Add(Tree);
 	}
@@ -312,32 +293,15 @@ void UUVSelectTool::Setup()
 			Targets[i]->UnwrapPreview->PreviewMesh->GetTransform());
 	}
 
-	// Remove trees that are no longer relevant (because their layer is not displayed)
-	TreeStore->RemoveByPredicate([this](TPair<FDynamicMesh3*, TSharedPtr<FDynamicMeshAABBTree3>> Pair)
-	{
-		for (TObjectPtr<UUVEditorToolMeshInput> Target : Targets)
-		{
-			if (Target->UnwrapCanonical.Get() == Pair.Key)
-			{
-				return false;
-			}
-		}
-		return true;
-	});
-
-	// Make sure that if we receive undo/redo events on the meshes, we update the tree structures
-	// and the selection mechanic drawn elements. Note that we mainly have to worry about this
-	// because the select tool is the default UV editor tool, and as such it can receive undo
-	// transactions from other tools and from other select tool invocations. Other tools typically
-	// only need to worry about their own transactions, since we undo other tool invocations before
-	// we get to unrelated transactions, and we can't redo out of the default tool.
+	// Make sure that if undo/redo events act on the meshes, we update our state.
+	// The trees will be updated by the tree store, which listens to the same broadcasts.
 	for (int32 i = 0; i < Targets.Num(); ++i)
 	{
-		Targets[i]->OnUndoRedo.AddWeakLambda(this, [this, i](bool bRevert) {
-			AABBTrees[i]->Build();
+		Targets[i]->OnCanonicalModified.AddWeakLambda(this, [this, i]
+		(UUVEditorToolMeshInput* InputObject, const UUVEditorToolMeshInput::FCanonicalModifiedInfo& Info) {
 			UpdateGizmo();
 			SelectionMechanic->RebuildDrawnElements(TransformGizmo->GetGizmoTransform());
-			});
+		});
 	}
 
 	// Gizmo setup
@@ -391,7 +355,7 @@ void UUVSelectTool::Shutdown(EToolShutdownType ShutdownType)
 
 	for (TObjectPtr<UUVEditorToolMeshInput> Target : Targets)
 	{
-		Target->OnUndoRedo.RemoveAll(this);
+		Target->OnCanonicalModified.RemoveAll(this);
 	}
 
 	Settings->SaveProperties(this);
@@ -675,11 +639,13 @@ void UUVSelectTool::GizmoTransformEnded(UTransformProxy* Proxy)
 	if (Settings->bUpdatePreviewDuringDrag)
 	{
 		// Both previews must already be updated, so only need to update canonical
-		Targets[SelectionTargetIndex]->UpdateCanonicalFromPreviews(&MovingVids);
+		Targets[SelectionTargetIndex]->UpdateCanonicalFromPreviews(&MovingVids, 
+			UUVEditorToolMeshInput::NONE_CHANGED_ARG);
 	}
 	else
 	{
-		Targets[SelectionTargetIndex]->UpdateAllFromUnwrapPreview(&MovingVids, nullptr, &SelectedTids);
+		Targets[SelectionTargetIndex]->UpdateAllFromUnwrapPreview(&MovingVids, 
+			UUVEditorToolMeshInput::NONE_CHANGED_ARG, &SelectedTids);
 	}
 
 	if (!AABBTrees[SelectionTargetIndex]->IsValid(false))
@@ -688,12 +654,10 @@ void UUVSelectTool::GizmoTransformEnded(UTransformProxy* Proxy)
 	}
 
 	const FText TransactionName(LOCTEXT("DragCompleteTransactionName", "Move Items"));
-	EmitChangeAPI->BeginUndoTransaction(TransactionName);
 	EmitChangeAPI->EmitToolIndependentChange(ChangeRouter, MakeUnique<UVSelectToolLocals::FGizmoMeshChange>(
 		Targets[SelectionTargetIndex], ChangeTracker.EndChange(), 
 		InitialGizmoFrame.ToFTransform(), TransformGizmo->GetGizmoTransform()),
 		TransactionName);
-	EmitChangeAPI->EndUndoTransaction();
 
 	TransformGizmo->SetNewChildScale(FVector::One());
 	SelectionMechanic->RebuildDrawnElements(TransformGizmo->GetGizmoTransform());
@@ -720,13 +684,15 @@ void UUVSelectTool::ApplyGizmoTransform()
 					MeshIn.SetVertex(MovingVids[i], TransformToApply.TransformPosition(MovingVertOriginalPositions[i]));
 				}
 			}, false);
-		Targets[SelectionTargetIndex]->UpdateUnwrapPreviewOverlayFromPositions(&MovingVids, nullptr, &SelectedTids);
+		Targets[SelectionTargetIndex]->UpdateUnwrapPreviewOverlayFromPositions(&MovingVids, 
+			UUVEditorToolMeshInput::NONE_CHANGED_ARG, &SelectedTids);
 
 		SelectionMechanic->SetDrawnElementsTransform((FTransform)TransformToApply);
 
 		if (Settings->bUpdatePreviewDuringDrag)
 		{
-			Targets[SelectionTargetIndex]->UpdateAppliedPreviewFromUnwrapPreview(&MovingVids, nullptr, &SelectedTids);
+			Targets[SelectionTargetIndex]->UpdateAppliedPreviewFromUnwrapPreview(&MovingVids, 
+				UUVEditorToolMeshInput::NONE_CHANGED_ARG, &SelectedTids);
 		}
 
 		bGizmoTransformNeedsApplication = false;
@@ -792,5 +758,4 @@ void UUVSelectTool::ApplyAction(ESelectToolAction ActionType)
 	}
 }
 
-
 #undef LOCTEXT_NAMESPACE
diff --git a/Engine/Plugins/Experimental/UVEditor/Source/UVEditorTools/Private/UVToolContextObjects.cpp b/Engine/Plugins/Experimental/UVEditor/Source/UVEditorTools/Private/UVToolContextObjects.cpp
--- a/Engine/Plugins/Experimental/UVEditor/Source/UVEditorTools/Private/UVToolContextObjects.cpp
+++ b/Engine/Plugins/Experimental/UVEditor/Source/UVEditorTools/Private/UVToolContextObjects.cpp
@@ -14,8 +14,8 @@ namespace UVToolContextObjectLocals
 {
 
 	/**
-	 * A wrapper change that applies a given change to the unwrap canonical mesh of an input, uses that to update the
-	 * other views, and issues an OnUndoRedo broadcast.
+	 * A wrapper change that applies a given change to the unwrap canonical mesh of an input, and uses that
+	 * to update the other views. Causes a broadcast of OnCanonicalModified.
 	 */
 	class  FUVEditorMeshChange : public FToolCommandChange
 	{
@@ -32,14 +32,12 @@ namespace UVToolContextObjectLocals
 		{
 			UnwrapCanonicalMeshChange->Apply(UVToolInputObject->UnwrapCanonical.Get(), false);
 			UVToolInputObject->UpdateFromCanonicalUnwrapUsingMeshChange(*UnwrapCanonicalMeshChange);
-			UVToolInputObject->OnUndoRedo.Broadcast(false);
 		}
 
 		virtual void Revert(UObject* Object) override
 		{
 			UnwrapCanonicalMeshChange->Apply(UVToolInputObject->UnwrapCanonical.Get(), true);
 			UVToolInputObject->UpdateFromCanonicalUnwrapUsingMeshChange(*UnwrapCanonicalMeshChange);
-			UVToolInputObject->OnUndoRedo.Broadcast(true);
 		}
 
 		virtual bool HasExpired(UObject* Object) const override
@@ -77,35 +75,62 @@ void UUVToolEmitChangeAPI::EmitToolDependentChange(UObject* TargetObject, TUniqu
 	ToolManager->EmitObjectChange(TargetObject, MoveTemp(Change), Description);
 }
 
-void UUVToolLivePreviewAPI::Initialize(UWorld* WorldIn, UInputRouter* RouterIn)
+void UUVToolLivePreviewAPI::Initialize(UWorld* WorldIn, UInputRouter* RouterIn,
+	TUniqueFunction<void(FViewCameraState& CameraStateOut)> GetLivePreviewCameraStateFuncIn)
 {
 	World = WorldIn;
 	InputRouter = RouterIn;
+	GetLivePreviewCameraStateFunc = MoveTemp(GetLivePreviewCameraStateFuncIn);
 }
 
-void UUVToolAABBTreeStorage::Set(FDynamicMesh3* MeshKey, TSharedPtr<FDynamicMeshAABBTree3> Tree)
+void UUVToolAABBTreeStorage::Set(FDynamicMesh3* MeshKey, TSharedPtr<FDynamicMeshAABBTree3> Tree, 
+	UUVEditorToolMeshInput* InputObject)
 {
-	AABBTrees.Add(MeshKey, Tree);
+	AABBTreeStorage.Add(MeshKey, TreeInputObjectPair(Tree, InputObject));
+	InputObject->OnCanonicalModified.AddWeakLambda(this, [MeshKey, Tree]
+	(UUVEditorToolMeshInput* InputObject, const UUVEditorToolMeshInput::FCanonicalModifiedInfo& Info) {
+		if (InputObject->UnwrapCanonical.Get() == MeshKey
+			|| (Info.bAppliedMeshShapeChanged && InputObject->AppliedCanonical.Get() == MeshKey))
+		{
+			Tree->Build();
+		}
+	});
 }
 
 TSharedPtr<FDynamicMeshAABBTree3> UUVToolAABBTreeStorage::Get(FDynamicMesh3* MeshKey)
 {
-	TSharedPtr<FDynamicMeshAABBTree3>* FoundPtr = AABBTrees.Find(MeshKey);
-	return FoundPtr ? *FoundPtr : nullptr;
+	TreeInputObjectPair* FoundPtr = AABBTreeStorage.Find(MeshKey);
+	return FoundPtr ? FoundPtr->Key : nullptr;
 }
 
 void UUVToolAABBTreeStorage::Remove(FDynamicMesh3* MeshKey)
 {
-	AABBTrees.Remove(MeshKey);
+	TreeInputObjectPair* Result = AABBTreeStorage.Find(MeshKey);
+	if (!Result)
+	{
+		return;
+	}
+	if (Result->Value.IsValid())
+	{
+		Result->Value->OnCanonicalModified.RemoveAll(this);
+	}
+
+	AABBTreeStorage.Remove(MeshKey);
 }
 
 void UUVToolAABBTreeStorage::RemoveByPredicate(TUniqueFunction<
-	bool(const TPair<FDynamicMesh3*, TSharedPtr<FDynamicMeshAABBTree3>>&)> Predicate)
+	bool(const FDynamicMesh3* Mesh, TWeakObjectPtr<UUVEditorToolMeshInput> InputObject,
+		TSharedPtr<FDynamicMeshAABBTree3> Tree)> Predicate)
 {
-	for (auto It = AABBTrees.CreateIterator(); It; ++It)
+	for (auto It = AABBTreeStorage.CreateIterator(); It; ++It)
 	{
-		if (Predicate(*It))
+		TWeakObjectPtr<UUVEditorToolMeshInput> InputObject = It->Value.Value;
+		if (Predicate(It->Key, InputObject, It->Value.Key)) // mesh, object, tree
 		{
+			if (InputObject.IsValid())
+			{
+				InputObject->OnCanonicalModified.RemoveAll(this);
+			}
 			It.RemoveCurrent();
 		}
 	}
@@ -113,5 +138,20 @@ void UUVToolAABBTreeStorage::RemoveByPredicate(TUniqueFunction<
 
 void UUVToolAABBTreeStorage::Empty()
 {
-	AABBTrees.Empty();
+	for (TPair<FDynamicMesh3*, TreeInputObjectPair>& MapPair : AABBTreeStorage)
+	{
+		TWeakObjectPtr<UUVEditorToolMeshInput> InputObject = MapPair.Value.Value;
+		if (InputObject.IsValid())
+		{
+			InputObject->OnCanonicalModified.RemoveAll(this);
+		}
+	}
+	AABBTreeStorage.Empty();
+}
+
+void UUVToolAABBTreeStorage::BeginDestroy()
+{
+	Empty();
+
+	Super::BeginDestroy();
 }
\ No newline at end of file
diff --git a/Engine/Plugins/Experimental/UVEditor/Source/UVEditorTools/Public/ToolTargets/UVEditorToolMeshInput.h b/Engine/Plugins/Experimental/UVEditor/Source/UVEditorTools/Public/ToolTargets/UVEditorToolMeshInput.h
--- a/Engine/Plugins/Experimental/UVEditor/Source/UVEditorTools/Public/ToolTargets/UVEditorToolMeshInput.h
+++ b/Engine/Plugins/Experimental/UVEditor/Source/UVEditorTools/Public/ToolTargets/UVEditorToolMeshInput.h
@@ -38,6 +38,7 @@ class UVEDITORTOOLS_API UUVEditorToolMeshInput : public UToolTarget
 	GENERATED_BODY()
 
 public:
+
 	/** 
 	 * Mesh representing the unwrapped UV layer. If the UnwrapPreview is changed via background
 	 * ops, then this mesh can be used to restart an operation as parameters change. Once a change
@@ -79,14 +80,34 @@ public:
 	UPROPERTY()
 	TObjectPtr<UMeshElementsVisualizer> WireframeDisplay = nullptr;
 
+	// Information about an OnCanonicalModified broadcast. This is a struct so that we can
+	// add more information if we someday need to.
+	struct FCanonicalModifiedInfo
+	{
+		// It's not yet clear whether we need either of these. bAppliedOverlayChanged seems
+		// like it would be always true whenever a change is emitted. bUnwrapMeshShapeChanged would
+		// seeminly only be false if we edited some other layer in the applied canonical without
+		// displaying it. That is rare and it would probably be safer to just do a full update
+		// in most of those cases (i.e. still assume that unwrap changed).
+		// So for now we don't use them.
+		//bool bUnwrapMeshShapeChanged = true;
+		//bool bAppliedOverlayChanged = true;
+
+		// This would only be true if we udpated the mesh after editing it externally, since
+		// the UV editor doesn't change mesh shape.
+		bool bAppliedMeshShapeChanged = false;
+	};
+
 	/**
-	 * Helper delegate that allows undo/redo transactions to notify the active tool that the input
-	 * object was modified. This is helpful for allowing transactions performed in one tool invocation
-	 * to properly be reverted/applied during a different (default) tool invocation, since the new
-	 * tool invocation can listen to this delegate and update relevant structures on broadcast.
+	 * Broadcast when the canonical unwrap or applied meshes change. This gets broadcast by
+	 * utility functions in this class that update those meshes (except for 
+	 * UpdateUnwrapCanonicalOverlayFromPositions, which is usually followed by a more complete
+	 * update). If those utility functions are not used and the client updates one of those meshes,
+	 * the client should broadcast the change themselves so that important related info can be 
+	 * updated (for instance, so that the mode can mark those objects as modified).
 	 */
-	DECLARE_MULTICAST_DELEGATE_OneParam(FOnUndoRedo, bool bRevert);
-	FOnUndoRedo OnUndoRedo;
+	DECLARE_MULTICAST_DELEGATE_TwoParams(FOnObjectModified, UUVEditorToolMeshInput* InputObject, const FCanonicalModifiedInfo&);
+	FOnObjectModified OnCanonicalModified;
 
 	// Additional needed information
 	TObjectPtr<UToolTarget> SourceTarget = nullptr;
@@ -109,48 +130,68 @@ public:
 	void Shutdown();
 
 	// Notes above the below convenience functions:
-	// 1. If BOTH ChangedVids/ChagnedElementIDs and ChangedConnectivityTids are null, a complete update will be performed. If either
-	//   is present, then only those vids/elements/tids will be updated.
-	// 2. Generally ChangedVids/ChangedElementIDs is allowed to have new elements, since it is natural to split UVs.
-	//   However, ChangedConnectivityTids must not have new Tids, because the tids form our correspondence between the unwrap mesh and
-	//   the original layer.
+	// 1. Passing nullptr for ChangedVids/ChagnedElementIDs/ChangedConnectivityTids means that all vids/elements/tids
+	//  need updating, respectively. Passing UUVEditorToolMeshInput::NONE_CHANGED_ARG is equivalent to passing a pointer 
+	//  to an empty array, and means that none changed. Otherwise only the vids/elements/tids in the pointed-to arrays
+	//  are iterated over in the updated mesh.
+	//  While passing these allows updates to be quicker, note that failing to include the relevant tids/vids/elements
+	//  can put the input object into an invalid state or hit a check, for instance if a triangle refers to a new vid
+	//  that was not added in.
+	// 2. ChangedVids/ChangedElementIDs is allowed to have new elements, since it is natural to split UVs. However, 
+	//  ChangedConnectivityTids must not have new Tids, because the tids form our correspondence between the unwrap mesh and
+	//  the original layer.
+	//  Also note that if you are gathering the changed vids/elements based off of triangles, you should gather them
+	//  from the post-change mesh/overlay so that you do not miss any added elements, and since removed elements are
+	//  captured by changed tri connectivity.
 	// 3. FastRenderUpdateTids is an optional list of triangles whose render data needs updating, which can allow for faster preview
 	//  updates if the render buffers are properly split apart. If provided, it should be a superset of ChangedConnectivityTids,
 	//  and should contain the one-ring triangles of ChangedVids.
 	// 4. If updating the preview objects, note that the functions do not try to cancel any active computations, so an active
-	//   computation could reset things once it completes.
+	//  computation could reset things once it completes.
+
+	// Can be passed in as a ChangedVids/ChangedElements/ChangedConnectivityTids arguments as an equivalent of
+	// passing a pointer to an empty array.
+	static const TArray<int32>* const NONE_CHANGED_ARG;
 
 	/**
 	 * Updates UnwrapPreview UV Overlay from UnwrapPreview vert positions. Issues a NotifyDeferredEditCompleted
 	 * for both positions and UVs.
 	 */
-	void UpdateUnwrapPreviewOverlayFromPositions(const TArray<int32>* ChangedVids = nullptr, const TArray<int32>* ChangedConnectivityTids = nullptr, const TArray<int32>* FastRenderUpdateTids = nullptr);
+	void UpdateUnwrapPreviewOverlayFromPositions(const TArray<int32>* ChangedVids = nullptr, 
+		const TArray<int32>* ChangedConnectivityTids = nullptr, const TArray<int32>* FastRenderUpdateTids = nullptr);
 
 	/**
 	 * Updates UnwrapCanonical UV Overlay from UnwrapCanonical vert positions. Issues a NotifyDeferredEditCompleted
 	 * for both positions and UVs.
+	 * 
+	 * Doesn't broadcast because it is expected that this call is followed by another call that updates the rest of 
+	 * the input object.
 	 */
-	void UpdateUnwrapCanonicalOverlayFromPositions(const TArray<int32>* ChangedVids = nullptr, const TArray<int32>* ChangedConnectivityTids = nullptr);
+	void UpdateUnwrapCanonicalOverlayFromPositions(const TArray<int32>* ChangedVids = nullptr, 
+		const TArray<int32>* ChangedConnectivityTids = nullptr);
 
 	/**
 	 * Updates the AppliedPreview from UnwrapPreview, without updating the non-preview meshes. Useful for updates during
 	 * a drag, etc, when we only care about updating the visible items.
 	 * Assumes that the overlay in UnwrapPreview is already updated.
 	 */
-	void UpdateAppliedPreviewFromUnwrapPreview(const TArray<int32>* ChangedVids = nullptr, const TArray<int32>* ChangedConnectivityTids = nullptr, const TArray<int32>* FastRenderUpdateTids = nullptr);
+	void UpdateAppliedPreviewFromUnwrapPreview(const TArray<int32>* ChangedVids = nullptr, 
+		const TArray<int32>* ChangedConnectivityTids = nullptr, const TArray<int32>* FastRenderUpdateTids = nullptr);
 
 	/**
 	 * Updates only the UnwrapPreview from AppliedPreview, without updating the non-preview meshes. Useful for transient
 	 * updates when we only care about updating the visible items.
 	 */
 	void UpdateUnwrapPreviewFromAppliedPreview(
-		const TArray<int32>* ChangedElementIDs = nullptr, const TArray<int32>* ChangedConnectivityTids = nullptr, const TArray<int32>* FastRenderUpdateTids = nullptr);
+		const TArray<int32>* ChangedElementIDs = nullptr, const TArray<int32>* ChangedConnectivityTids = nullptr, 
+		const TArray<int32>* FastRenderUpdateTids = nullptr);
 
 	/**
 	 * Updates the non-preview meshes from their preview counterparts. Useful, for instance, after the completion of a drag
 	 * to update the canonical objects.
 	 */
-	void UpdateCanonicalFromPreviews(const TArray<int32>* ChangedVids = nullptr, const TArray<int32>* ChangedConnectivityTids = nullptr);
+	void UpdateCanonicalFromPreviews(const TArray<int32>* ChangedVids = nullptr, 
+		const TArray<int32>* ChangedConnectivityTids = nullptr, bool bBroadcast = true);
 
 	/**
 	 * Updates the preview meshes from their canonical counterparts. Useful mainly as a way to reset the previews.
@@ -163,35 +204,47 @@ public:
 	 * Used to update everything when the applied preview was not being updated in tandem with the unwrap preview (otherwise one 
 	 * would use UpdateCanonicalFromPreviews).
 	 */
-	void UpdateAllFromUnwrapPreview(const TArray<int32>* ChangedVids = nullptr, const TArray<int32>* ChangedConnectivityTids = nullptr, const TArray<int32>* FastRenderUpdateTids = nullptr);
+	void UpdateAllFromUnwrapPreview(const TArray<int32>* ChangedVids = nullptr, const TArray<int32>* ChangedConnectivityTids = nullptr, 
+		const TArray<int32>* FastRenderUpdateTids = nullptr, bool bBroadcast = true);
 
 	/**
 	 * Updates the other meshes using the mesh in UnwrapCanonical. Useful when immediately applying an operation that
 	 * operates directly on the unwrap mesh.
 	 */
-	void UpdateAllFromUnwrapCanonical(const TArray<int32>* ChangedVids = nullptr, const TArray<int32>* ChangedConnectivityTids = nullptr, const TArray<int32>* FastRenderUpdateTids = nullptr);
+	void UpdateAllFromUnwrapCanonical(const TArray<int32>* ChangedVids = nullptr, const TArray<int32>* ChangedConnectivityTids = nullptr, 
+		const TArray<int32>* FastRenderUpdateTids = nullptr, bool bBroadcast = true);
 
 	/**
 	 * Updates the other meshes using the mesh in AppliedCanonical. Useful when switching UV layer indices or otherwise
 	 * resetting the collection from the "original" data.
 	 */
-	void UpdateAllFromAppliedCanonical(const TArray<int32>* ChangedVids = nullptr, const TArray<int32>* ChangedConnectivityTids = nullptr, const TArray<int32>* FastRenderUpdateTids = nullptr);
+	void UpdateAllFromAppliedCanonical(const TArray<int32>* ChangedVids = nullptr, const TArray<int32>* ChangedConnectivityTids = nullptr, 
+		const TArray<int32>* FastRenderUpdateTids = nullptr, bool bBroadcast = true);
 
 	/**
 	 * Updates the other meshes using the UV overlay in the live preview.
 	 */
-	void UpdateAllFromAppliedPreview(const TArray<int32>* ChangedElementIDs = nullptr, const TArray<int32>* ChangedConnectivityTids = nullptr, const TArray<int32>* FastRenderUpdateTids = nullptr);
+	void UpdateAllFromAppliedPreview(const TArray<int32>* ChangedElementIDs = nullptr, 
+		const TArray<int32>* ChangedConnectivityTids = nullptr, const TArray<int32>* FastRenderUpdateTids = nullptr, bool bBroadcast = true);
 
 	/**
 	 * Uses the stored triangles/vertices in the mesh change to update everything from the canonical unwrap. 
 	 * Assumes the change has already been applied to the canonical unwrap.
 	 */
-	void UpdateFromCanonicalUnwrapUsingMeshChange(const UE::Geometry::FDynamicMeshChange& UnwrapCanonicalMeshChange);
+	void UpdateFromCanonicalUnwrapUsingMeshChange(const UE::Geometry::FDynamicMeshChange& UnwrapCanonicalMeshChange, bool bBroadcast = true);
+
+	/**
+	 * Convert a vid in the unwrap mesh to the corresponding vid in the applied mesh (i.e., parent vertex of the
+	 * element corresponding to the unwrap vertex).
+	 */
+	int32 UnwrapVidToAppliedVid(int32 UnwrapVid);
+
+	/**
+	 * Get the unwrap vids corresponding to a given applied vid. These will be multiple if the vertex is
+	 * a seam vertex and therefore has multiple UV elements associated with it.
+	 */
+	void AppliedVidToUnwrapVids(int32 AppliedVid, TArray<int32>& UnwrapVidsOut);
 
 	// UToolTarget
 	virtual bool IsValid() const override;
-protected:
-
-	void UpdateOtherUnwrap(const UE::Geometry::FDynamicMesh3& SourceUnwrapMesh,
-		UE::Geometry::FDynamicMesh3& DestUnwrapMesh, const TArray<int32>* ChangedVids, const TArray<int32>* ChangedConnectivityTids);
 };
\ No newline at end of file
diff --git a/Engine/Plugins/Experimental/UVEditor/Source/UVEditorTools/Public/UVEditorToolUtil.h b/Engine/Plugins/Experimental/UVEditor/Source/UVEditorTools/Public/UVEditorToolUtil.h
--- a/Engine/Plugins/Experimental/UVEditor/Source/UVEditorTools/Public/UVEditorToolUtil.h
+++ b/Engine/Plugins/Experimental/UVEditor/Source/UVEditorTools/Public/UVEditorToolUtil.h
@@ -15,16 +15,61 @@ namespace UE {
 namespace Geometry {
 namespace UVEditorToolUtil {
 
+	/** 
+	 * Given a UV overlay, generate a mesh out of the overlay using the given UV to 3d position function.
+	 * The resulting mesh's vertex IDs will match the overlay element IDs, and the triangle IDs will match
+	 * the overlay triangle IDs.
+	 */
 	UVEDITORTOOLS_API void GenerateUVUnwrapMesh(const FDynamicMeshUVOverlay& UVOverlay, FDynamicMesh3& UnwrapMeshOut,
 		TFunctionRef<FVector3d(const FVector2f&)> UVToVertPosition);
 
+	/**
+	 * Given a mesh that corresponds to an overlay with an exact mapping of triangle IDs and element IDs to
+	 * vertex IDs (i.e. one generated by GenerateUVUnwrapMesh), update the vertex locations and triangle
+	 * connectivity of the mesh.
+	 * 
+	 * @param ChangedElementIDs If not null, only update vertex locations corresponding to the given elements.
+	 * @param ChangedConnectivityTids If not null, only update triangle connectivity corresponding to the given triangle IDs.
+	 */
 	UVEDITORTOOLS_API void UpdateUVUnwrapMesh(const FDynamicMeshUVOverlay& UVOverlayIn,
 		FDynamicMesh3& UnwrapMeshOut, TFunctionRef<FVector3d(const FVector2f&)> UVToVertPosition,
-		const TArray<int32>* ChangedElementIDs, const TArray<int32>* ChangedTids);
+		const TArray<int32>* ChangedElementIDs, const TArray<int32>* ChangedConnectivityTids);
 
+	/**
+	 * Given a copy of a UV Unwrap mesh whose positions and connectivity have been updated, update the
+	 * destination unwrap mesh.
+	 *
+	 * @param ChangedVids If not null, only update elements corresponding to the given vertex IDs.
+	 * @param ChangedConnectivityTids If not null, only update triangle connectivity corresponding to the given triangle IDs.
+	 */
+	UVEDITORTOOLS_API void UpdateUVUnwrapMesh(const FDynamicMesh3& SourceUnwrapMesh, FDynamicMesh3& DestUnwrapMesh,
+		const TArray<int32>* ChangedVids, const TArray<int32>* ChangedConnectivityTids);
+
+	/**
+	 * Given a mesh that corresponds to an overlay with an exact mapping of triangle IDs and element IDs to
+	 * vertex IDs (i.e. one generated by GenerateUVUnwrapMesh), update the element values and connectivity
+	 * of the overlay.
+	 * 
+	 * @param ChangedVids If not null, only update elements corresponding to the given vertex IDs.
+	 * @param ChangedConnectivityTids If not null, only update triangle connectivity corresponding to the given triangle IDs.
+	 */
 	UVEDITORTOOLS_API void UpdateUVOverlayFromUnwrapMesh(
 		const FDynamicMesh3& UnwrapMeshIn, FDynamicMeshUVOverlay& UVOverlayOut,
 		TFunctionRef<FVector2f(const FVector3d&)> UVToVertPosition,
-		const TArray<int32>* ChangedVids = nullptr, const TArray<int32>* ChangedTids = nullptr);
+		const TArray<int32>* ChangedVids = nullptr, const TArray<int32>* ChangedConnectivityTids = nullptr);
+
+	/**
+	 * Given two overlays that have the same triangle IDs, copy all or part of one overlay into another.
+	 *
+	 * @param bSameParentVerts If true, underlying meshes have same connectivity/id's, meaning that overlays
+	 *   have same parent verts, and we can use a simple copy when updating everything. Can be set true when 
+	 *   copying between a canonical/preview version of the same mesh, but must be false when copying between
+	 *   an unwrap mesh and an applied mesh.
+	 * @param ChangedVids If not null, only update elements corresponding to the given vertex IDs.
+	 * @param ChangedConnectivityTids If not null, only update triangle connectivity corresponding to the given triangle IDs.
+	 */
+	UVEDITORTOOLS_API void UpdateOverlayFromOverlay(const FDynamicMeshUVOverlay& OverlayIn, FDynamicMeshUVOverlay& OverlayOut,
+		bool bParentVertsIdentical, const TArray<int32>* ChangedElements = nullptr,
+		const TArray<int32>* ChangedConnectivityTids = nullptr);
 
-}}}
\ No newline at end of file
+}}}//end UE::Geometry::UVEditorToolUtil
\ No newline at end of file
diff --git a/Engine/Plugins/Experimental/UVEditor/Source/UVEditorTools/Public/UVToolContextObjects.h b/Engine/Plugins/Experimental/UVEditor/Source/UVEditorTools/Public/UVToolContextObjects.h
--- a/Engine/Plugins/Experimental/UVEditor/Source/UVEditorTools/Public/UVToolContextObjects.h
+++ b/Engine/Plugins/Experimental/UVEditor/Source/UVEditorTools/Public/UVToolContextObjects.h
@@ -16,6 +16,7 @@
 
 PREDECLARE_GEOMETRY(class FDynamicMesh3);
 class FToolCommandChange;
+struct FViewCameraState;
 class UInputRouter;
 class UWorld;
 class UUVEditorToolMeshInput;
@@ -92,16 +93,27 @@ class UVEDITORTOOLS_API UUVToolLivePreviewAPI : public UUVToolContextObject
 	GENERATED_BODY()
 public:
 
-	void Initialize(UWorld* WorldIn, UInputRouter* RouterIn);
+	void Initialize(UWorld* WorldIn, UInputRouter* RouterIn,
+		TUniqueFunction<void(FViewCameraState& CameraStateOut)> GetLivePreviewCameraStateFuncIn);
 
 	UWorld* GetLivePreviewWorld() { return World.Get(); }
 	UInputRouter* GetLivePreviewInputRouter() { return InputRouter.Get(); }
+	void GetLivePreviewCameraState(FViewCameraState& CameraStateOut) 
+	{ 
+		if (GetLivePreviewCameraStateFunc)
+		{
+			GetLivePreviewCameraStateFunc(CameraStateOut);
+		}
+	}
+
 protected:
 	UPROPERTY()
 	TWeakObjectPtr<UWorld> World;
 
 	UPROPERTY()
 	TWeakObjectPtr<UInputRouter> InputRouter;
+
+	TUniqueFunction<void(FViewCameraState& CameraStateOut)> GetLivePreviewCameraStateFunc;
 };
 
 /**
@@ -191,25 +203,39 @@ public:
 
 };
 
-/** Stores UV mesh AABB trees */
+/** 
+ * Stores AABB trees for UV input object unwrap canonical or applied canonical meshes.
+ * Binds to the input objects' OnCanonicalModified delegate to automatically update 
+ * the tree when necessary.
+ */
 UCLASS()
 class UVEDITORTOOLS_API UUVToolAABBTreeStorage : public UUVToolContextObject
 {
 	GENERATED_BODY()
 public:
+	using FDynamicMesh3 = UE::Geometry::FDynamicMesh3;
+	using FDynamicMeshAABBTree3 = UE::Geometry::FDynamicMeshAABBTree3;
+
 
-	void Set(UE::Geometry::FDynamicMesh3* MeshKey, TSharedPtr<UE::Geometry::FDynamicMeshAABBTree3> Tree);
+	void Set(FDynamicMesh3* MeshKey, TSharedPtr<FDynamicMeshAABBTree3> Tree,
+		UUVEditorToolMeshInput* InputObject);
 
-	TSharedPtr<UE::Geometry::FDynamicMeshAABBTree3> Get(UE::Geometry::FDynamicMesh3* MeshKey);
+	TSharedPtr<FDynamicMeshAABBTree3> Get(FDynamicMesh3* MeshKey);
 
-	void Remove(UE::Geometry::FDynamicMesh3* MeshKey);
+	void Remove(FDynamicMesh3* MeshKey);
 
 	void RemoveByPredicate(TUniqueFunction<
-		bool(const TPair<UE::Geometry::FDynamicMesh3*, TSharedPtr<UE::Geometry::FDynamicMeshAABBTree3>>&)> Predicate);
+		bool(const FDynamicMesh3* Mesh, TWeakObjectPtr<UUVEditorToolMeshInput> InputObject,
+			TSharedPtr<FDynamicMeshAABBTree3> Tree)> Predicate);
 
 	void Empty();
 
+	// UObject
+	virtual void BeginDestroy() override;
+
 protected:
 
-	TMap<UE::Geometry::FDynamicMesh3*, TSharedPtr<UE::Geometry::FDynamicMeshAABBTree3>> AABBTrees;
+	using TreeInputObjectPair = TPair<TSharedPtr<FDynamicMeshAABBTree3>, TWeakObjectPtr<UUVEditorToolMeshInput>>;
+
+	TMap<FDynamicMesh3*, TreeInputObjectPair> AABBTreeStorage;
 };
