UVEditor: Minor fix to seam sew action passing in duplicate changed id's.

#rnx

diff --git a/Engine/Plugins/Experimental/UVEditor/Source/UVEditorTools/Private/UVSeamSewAction.cpp b/Engine/Plugins/Experimental/UVEditor/Source/UVEditorTools/Private/UVSeamSewAction.cpp
--- a/Engine/Plugins/Experimental/UVEditor/Source/UVEditorTools/Private/UVSeamSewAction.cpp
+++ b/Engine/Plugins/Experimental/UVEditor/Source/UVEditorTools/Private/UVSeamSewAction.cpp
@@ -253,14 +253,14 @@ bool UUVSeamSewAction::ApplySewAction(UUVToolEmitChangeAPI& EmitChangeAPI)
 		}
 	}
 	SelectedTids.Sort();
-	Algo::Unique(SelectedTids);
+	SelectedTids.SetNum(Algo::Unique(SelectedTids));
 
 	FDynamicMeshChangeTracker ChangeTracker(Targets[SelectionTargetIndex]->UnwrapCanonical.Get());
 	ChangeTracker.BeginChange();
 	ChangeTracker.SaveTriangles(SelectedTids, true);
 
 	EdgeSewCandidates.Sort();
-	Algo::Unique(EdgeSewCandidates);
+	EdgeSewCandidates.SetNum(Algo::Unique(EdgeSewCandidates));
 
 	TArray<FDynamicMesh3::FMergeEdgesInfo> AllMergeInfo;
 	
@@ -277,7 +277,7 @@ bool UUVSeamSewAction::ApplySewAction(UUVToolEmitChangeAPI& EmitChangeAPI)
 			UE_LOG(LogGeometry, Warning, TEXT("Failed to sew edge pair %d / %d. Failed with code %d"), EdgePair[0], EdgePair[1], Result);
 		}
 	}
-	checkSlow(MeshToSew.CheckValidity());
+	checkSlow(MeshToSew.CheckValidity(FDynamicMesh3::FValidityOptions(true, true))); // Allow nonmanifold verts and reverse orientation
 
 
 	TArray<int32> RemainingVids;
