UVEditor: Fix 3D viewport to not display transform gizmos and to provide better controls over camera movement

#rnx

diff --git a/Engine/Plugins/Experimental/UVEditor/Source/UVEditor/Private/SUVEditor2DViewportToolBar.cpp b/Engine/Plugins/Experimental/UVEditor/Source/UVEditor/Private/SUVEditor2DViewportToolBar.cpp
--- a/Engine/Plugins/Experimental/UVEditor/Source/UVEditor/Private/SUVEditor2DViewportToolBar.cpp
+++ b/Engine/Plugins/Experimental/UVEditor/Source/UVEditor/Private/SUVEditor2DViewportToolBar.cpp
@@ -43,7 +43,7 @@ TSharedRef<SWidget> SUVEditor2DViewportToolBar::MakeToolBar(const TSharedPtr<FEx
 	// The following is modeled after portions of STransformViewportToolBar, which gets 
 	// used in SCommonEditorViewportToolbarBase.
 
-	// The buttons are hooked up to actual functions via command bindings in SUVEditor2DViewport::BindCommands(),
+	// The buttons are hooked up to actual functions via command bindings in SEditorViewport::BindCommands(),
 	// and the toolbar gets built in SUVEditor2DViewport::MakeViewportToolbar().
 
 	FSlimHorizontalToolBarBuilder ToolbarBuilder(CommandList, FMultiBoxCustomization::None, InExtenders);
diff --git a/Engine/Plugins/Experimental/UVEditor/Source/UVEditor/Private/SUVEditor3DViewport.cpp b/Engine/Plugins/Experimental/UVEditor/Source/UVEditor/Private/SUVEditor3DViewport.cpp
new file mode 100644
--- /dev/null
+++ b/Engine/Plugins/Experimental/UVEditor/Source/UVEditor/Private/SUVEditor3DViewport.cpp
@@ -0,0 +1,41 @@
+// Copyright Epic Games, Inc. All Rights Reserved.
+
+#include "SUVEditor3DViewport.h"
+
+#include "EditorViewportClient.h"
+#include "SUVEditor3DViewportToolBar.h"
+#include "UVEditorCommands.h"
+#include "UVEditor3DViewportClient.h"
+
+#define LOCTEXT_NAMESPACE "SUVEditor3DViewport"
+
+void SUVEditor3DViewport::BindCommands()
+{
+	SAssetEditorViewport::BindCommands();
+
+	const FUVEditorCommands& CommandInfos = FUVEditorCommands::Get();
+	CommandList->MapAction(
+		CommandInfos.EnableOrbitCamera,
+		FExecuteAction::CreateLambda([this]() {
+			StaticCastSharedPtr<FUVEditor3DViewportClient>(Client)->SetCameraMode(EUVEditor3DViewportClientCameraMode::Orbit);
+			}),
+		FCanExecuteAction::CreateLambda([this]() { return true; }),
+		FIsActionChecked::CreateLambda([this]() { return StaticCastSharedPtr<FUVEditor3DViewportClient>(Client)->GetCameraMode() == EUVEditor3DViewportClientCameraMode::Orbit; }));
+
+	CommandList->MapAction(
+		CommandInfos.EnableFlyCamera,
+		FExecuteAction::CreateLambda([this]() {
+			StaticCastSharedPtr<FUVEditor3DViewportClient>(Client)->SetCameraMode(EUVEditor3DViewportClientCameraMode::Fly);
+			}),
+		FCanExecuteAction::CreateLambda([this]() { return true; }),
+		FIsActionChecked::CreateLambda([this]() { return StaticCastSharedPtr<FUVEditor3DViewportClient>(Client)->GetCameraMode() == EUVEditor3DViewportClientCameraMode::Fly; }));
+}
+
+TSharedPtr<SWidget> SUVEditor3DViewport::MakeViewportToolbar()
+{
+	return SNew(SUVEditor3DViewportToolBar)
+		.CommandList(CommandList);
+}
+
+
+#undef LOCTEXT_NAMESPACE
diff --git a/Engine/Plugins/Experimental/UVEditor/Source/UVEditor/Private/SUVEditor3DViewportToolBar.cpp b/Engine/Plugins/Experimental/UVEditor/Source/UVEditor/Private/SUVEditor3DViewportToolBar.cpp
new file mode 100644
--- /dev/null
+++ b/Engine/Plugins/Experimental/UVEditor/Source/UVEditor/Private/SUVEditor3DViewportToolBar.cpp
@@ -0,0 +1,78 @@
+// Copyright Epic Games, Inc. All Rights Reserved.
+
+#pragma once
+
+#include "SUVEditor3DViewportToolBar.h"
+
+#include "UVEditorCommands.h"
+#include "Framework/MultiBox/MultiBoxBuilder.h"
+#include "Widgets/Layout/SBorder.h"
+#include "Widgets/SBoxPanel.h"
+
+
+void SUVEditor3DViewportToolBar::Construct(const FArguments& InArgs)
+{
+	CommandList = InArgs._CommandList;
+
+	const FMargin ToolbarSlotPadding(4.0f, 1.0f);
+	TSharedPtr<SHorizontalBox> MainBoxPtr;
+
+	ChildSlot
+	[
+		SNew( SBorder )
+		.BorderImage(FAppStyle::Get().GetBrush("EditorViewportToolBar.Background"))
+		.Cursor(EMouseCursor::Default)
+		[
+			SAssignNew( MainBoxPtr, SHorizontalBox )
+		]
+	];
+
+	MainBoxPtr->AddSlot()
+		.Padding(ToolbarSlotPadding)
+		.HAlign(HAlign_Right)
+		[
+			MakeToolBar(InArgs._Extenders)
+		];
+
+	SViewportToolBar::Construct(SViewportToolBar::FArguments());
+}
+
+TSharedRef<SWidget> SUVEditor3DViewportToolBar::MakeToolBar(const TSharedPtr<FExtender> InExtenders)
+{
+	// The following is modeled after portions of STransformViewportToolBar, which gets 
+	// used in SCommonEditorViewportToolbarBase.
+
+	// The buttons are hooked up to actual functions via command bindings in SUVEditor3DViewport::BindCommands(),
+	// and the toolbar gets built in SUVEditor3DViewport::MakeViewportToolbar().
+
+	FSlimHorizontalToolBarBuilder ToolbarBuilder(CommandList, FMultiBoxCustomization::None, InExtenders);
+
+	//// Use a custom style
+	FName ToolBarStyle = "EditorViewportToolBar";
+	ToolbarBuilder.SetStyle(&FEditorStyle::Get(), ToolBarStyle);
+	ToolbarBuilder.SetLabelVisibility(EVisibility::Collapsed);
+
+	ToolbarBuilder.BeginSection("OrbitFlyToggle");
+	{
+		ToolbarBuilder.BeginBlockGroup();
+	    
+		// TODO: Right now we're (sort-of) hardcoding the icons in here. We should have a style set for the uv
+		// editor that sets the correct icons for these.
+
+		// Orbit Camera
+		static FName OrbitCameraName = FName(TEXT("OrbitCamera"));
+		ToolbarBuilder.AddToolBarButton(FUVEditorCommands::Get().EnableOrbitCamera, NAME_None, TAttribute<FText>(), TAttribute<FText>(), 
+			TAttribute<FSlateIcon>(FSlateIcon(FEditorStyle::GetStyleSetName(), FName("EditorViewport.RotateMode"))), OrbitCameraName);
+
+		// Fly Camera
+		static FName FlyCameraName = FName(TEXT("FlyCamera"));
+		ToolbarBuilder.AddToolBarButton(FUVEditorCommands::Get().EnableFlyCamera, NAME_None, TAttribute<FText>(), TAttribute<FText>(), 
+			TAttribute<FSlateIcon>(FSlateIcon(FEditorStyle::GetStyleSetName(), FName("EditorViewport.CamSpeedSetting"))), FlyCameraName);
+
+		ToolbarBuilder.EndBlockGroup();
+	}
+
+	ToolbarBuilder.EndSection();
+
+	return ToolbarBuilder.MakeWidget();
+}
\ No newline at end of file
diff --git a/Engine/Plugins/Experimental/UVEditor/Source/UVEditor/Private/UVEditorCommands.cpp b/Engine/Plugins/Experimental/UVEditor/Source/UVEditor/Private/UVEditorCommands.cpp
--- a/Engine/Plugins/Experimental/UVEditor/Source/UVEditor/Private/UVEditorCommands.cpp
+++ b/Engine/Plugins/Experimental/UVEditor/Source/UVEditor/Private/UVEditorCommands.cpp
@@ -31,6 +31,9 @@ void FUVEditorCommands::RegisterCommands()
 	// These allow us to link up to pressed keys
 	UI_COMMAND(AcceptActiveTool, "Accept", "Accept the active tool", EUserInterfaceActionType::Button, FInputChord(EKeys::Enter));
 	UI_COMMAND(CancelActiveTool, "Cancel", "Cancel the active tool", EUserInterfaceActionType::Button, FInputChord(EKeys::Escape));
+
+	UI_COMMAND(EnableOrbitCamera, "Orbit", "Enable Orbit Camera", EUserInterfaceActionType::ToggleButton, FInputChord());
+	UI_COMMAND(EnableFlyCamera, "Fly", "Enable Fly Camera", EUserInterfaceActionType::ToggleButton, FInputChord());
 }
 
 #undef LOCTEXT_NAMESPACE
diff --git a/Engine/Plugins/Experimental/UVEditor/Source/UVEditor/Private/UVEditorToolkit.cpp b/Engine/Plugins/Experimental/UVEditor/Source/UVEditor/Private/UVEditorToolkit.cpp
--- a/Engine/Plugins/Experimental/UVEditor/Source/UVEditor/Private/UVEditorToolkit.cpp
+++ b/Engine/Plugins/Experimental/UVEditor/Source/UVEditor/Private/UVEditorToolkit.cpp
@@ -10,6 +10,7 @@
 #include "PreviewScene.h"
 #include "SAssetEditorViewport.h"
 #include "SUVEditor2DViewport.h"
+#include "SUVEditor3DViewport.h"
 #include "ThumbnailRendering/SceneThumbnailInfo.h"
 #include "ToolContextInterfaces.h"
 #include "ToolMenus.h"
@@ -18,6 +19,7 @@
 #include "UVEditorCommands.h"
 #include "UVEditorSubsystem.h"
 #include "UVEditor2DViewportClient.h"
+#include "UVEditor3DViewportClient.h"
 #include "UVToolContextObjects.h"
 #include "Widgets/Docking/SDockTab.h"
 
@@ -82,7 +84,7 @@ FUVEditorToolkit::FUVEditorToolkit(UAssetEditor* InOwningAssetEditor)
 					FTabManager::NewStack()
 					->SetSizeCoefficient(0.4f)
 					->AddTab(LivePreviewTabID, ETabState::OpenedTab)
-					->SetHideTabWell(false)
+					->SetHideTabWell(true)
 				)
 			)
 		);
@@ -99,12 +101,12 @@ FUVEditorToolkit::FUVEditorToolkit(UAssetEditor* InOwningAssetEditor)
 	LivePreviewInputRouter = LivePreviewEditorModeManager->GetInteractiveToolsContext()->InputRouter;
 
 	LivePreviewTabContent = MakeShareable(new FEditorViewportTabContent());
-	LivePreviewViewportClient = MakeShared<FEditorViewportClient>(
+	LivePreviewViewportClient = MakeShared<FUVEditor3DViewportClient>(
 		LivePreviewEditorModeManager.Get(), LivePreviewScene.Get());
 
 	LivePreviewViewportDelegate = [this](FAssetEditorViewportConstructionArgs InArgs)
 	{
-		return SNew(SAssetEditorViewport, InArgs)
+		return SNew(SUVEditor3DViewport, InArgs)
 			.EditorViewportClient(LivePreviewViewportClient);
 	};
 }
@@ -185,13 +187,13 @@ void FUVEditorToolkit::RegisterTabSpawners(const TSharedRef<FTabManager>& InTabM
 		.SetIcon(FSlateIcon(FEditorStyle::GetStyleSetName(), "LevelEditor.Tabs.Details"));
 
 	InTabManager->RegisterTabSpawner(ViewportTabID, FOnSpawnTab::CreateSP(this, &FUVEditorToolkit::SpawnTab_Viewport))
-		.SetDisplayName(LOCTEXT("ViewportTab", "Viewport"))
+		.SetDisplayName(LOCTEXT("ViewportTabUV", "UV Viewport"))
 		.SetGroup(AssetEditorTabsCategory.ToSharedRef())
 		.SetIcon(FSlateIcon(FEditorStyle::GetStyleSetName(), "LevelEditor.Tabs.Viewports"));
 
 	InTabManager->RegisterTabSpawner(LivePreviewTabID, FOnSpawnTab::CreateSP(this, 
 		&FUVEditorToolkit::SpawnTab_LivePreview))
-		.SetDisplayName(LOCTEXT("ViewportTab", "Viewport"))
+		.SetDisplayName(LOCTEXT("ViewportTabLivePreview", "3D Viewport"))
 		.SetGroup(AssetEditorTabsCategory.ToSharedRef())
 		.SetIcon(FSlateIcon(FEditorStyle::GetStyleSetName(), "LevelEditor.Tabs.Viewports"));
 }
@@ -433,10 +435,11 @@ void FUVEditorToolkit::PostInitAssetEditor()
 
 	// Adjust our live preview (3D) viewport
 	SetCommonViewportClientOptions(LivePreviewViewportClient.Get());
+	LivePreviewViewportClient->ToggleOrbitCamera(true);
+
 	// TODO: This should not be hardcoded
 	LivePreviewViewportClient->SetViewLocation(FVector(-200, 100, 100));
 	LivePreviewViewportClient->SetLookAtLocation(FVector(0, 0, 0));
-	LivePreviewViewportClient->ToggleOrbitCamera(true);
 }
 
 #undef LOCTEXT_NAMESPACE
\ No newline at end of file
diff --git a/Engine/Plugins/Experimental/UVEditor/Source/UVEditor/Public/SUVEditor3DViewport.h b/Engine/Plugins/Experimental/UVEditor/Source/UVEditor/Public/SUVEditor3DViewport.h
new file mode 100644
--- /dev/null
+++ b/Engine/Plugins/Experimental/UVEditor/Source/UVEditor/Public/SUVEditor3DViewport.h
@@ -0,0 +1,20 @@
+// Copyright Epic Games, Inc. All Rights Reserved.
+
+#pragma once
+
+#include "CoreMinimal.h"
+
+#include "SAssetEditorViewport.h"
+#include "SUVEditor3DViewportToolBar.h"
+
+/**
+ * Viewport used for live preview in UV editor. Has a custom toolbar overlay at the top.
+ */
+class UVEDITOR_API SUVEditor3DViewport : public SAssetEditorViewport
+{
+public:
+
+	// SAssetEditorViewport
+	virtual void BindCommands() override;
+	virtual TSharedPtr<SWidget> MakeViewportToolbar() override;
+};
diff --git a/Engine/Plugins/Experimental/UVEditor/Source/UVEditor/Public/SUVEditor3DViewportToolBar.h b/Engine/Plugins/Experimental/UVEditor/Source/UVEditor/Public/SUVEditor3DViewportToolBar.h
new file mode 100644
--- /dev/null
+++ b/Engine/Plugins/Experimental/UVEditor/Source/UVEditor/Public/SUVEditor3DViewportToolBar.h
@@ -0,0 +1,30 @@
+// Copyright Epic Games, Inc. All Rights Reserved.
+
+#pragma once
+
+#include "CoreMinimal.h"
+
+#include "SViewportToolBar.h"
+
+class FExtender;
+class FUICommandList;
+class SUVEditor3DViewport;
+
+/**
+ * Toolbar that shows up at the top of the 3d viewport (has camera controls)
+ */
+class SUVEditor3DViewportToolBar : public SViewportToolBar
+{
+public:
+	SLATE_BEGIN_ARGS(SUVEditor3DViewportToolBar) {}
+		SLATE_ARGUMENT(TSharedPtr<FUICommandList>, CommandList)
+		SLATE_ARGUMENT(TSharedPtr<FExtender>, Extenders)
+	SLATE_END_ARGS()
+
+	void Construct(const FArguments& InArgs);
+
+private:
+	TSharedRef<SWidget> MakeToolBar(const TSharedPtr<FExtender> InExtenders);
+
+	TSharedPtr<FUICommandList> CommandList;
+};
diff --git a/Engine/Plugins/Experimental/UVEditor/Source/UVEditor/Public/UVEditor3DViewportClient.h b/Engine/Plugins/Experimental/UVEditor/Source/UVEditor/Public/UVEditor3DViewportClient.h
new file mode 100644
--- /dev/null
+++ b/Engine/Plugins/Experimental/UVEditor/Source/UVEditor/Public/UVEditor3DViewportClient.h
@@ -0,0 +1,51 @@
+// Copyright Epic Games, Inc. All Rights Reserved.
+
+#pragma once
+
+#include "CoreMinimal.h"
+
+#include "EditorViewportClient.h"
+
+
+// Types of camera motion for the UV Editor 3D viewport
+enum EUVEditor3DViewportClientCameraMode {
+	Orbit,
+	Fly
+};
+
+/**
+ * Viewport client for the 3d live preview in the UV editor. Currently same as editor viewport
+ * client but doesn't allow editor gizmos/widgets, and alters orbit camera control.
+ */
+class UVEDITOR_API FUVEditor3DViewportClient : public FEditorViewportClient
+{
+	// Inherit all the same constructors
+	using FEditorViewportClient::FEditorViewportClient;
+
+	// FEditorViewportClient
+	virtual bool ShouldOrbitCamera() const override {
+		switch (CameraMode) {
+		case EUVEditor3DViewportClientCameraMode::Orbit:
+			return true;
+		case EUVEditor3DViewportClientCameraMode::Fly:
+			return FEditorViewportClient::ShouldOrbitCamera();
+		default:
+			ensure(false);
+			return FEditorViewportClient::ShouldOrbitCamera();
+		}
+	}
+	bool CanSetWidgetMode(UE::Widget::EWidgetMode NewMode) const override {	return false; }
+	void SetWidgetMode(UE::Widget::EWidgetMode NewMode) override {}
+	UE::Widget::EWidgetMode GetWidgetMode() const override { return UE::Widget::EWidgetMode::WM_None; }
+
+public:
+
+	void SetCameraMode(EUVEditor3DViewportClientCameraMode CameraModeIn) { CameraMode = CameraModeIn; };
+	EUVEditor3DViewportClientCameraMode GetCameraMode() const { return CameraMode; };
+
+protected:
+
+	// Enforce Orbit camera for UV editor live preview viewport. Use this instead of the base class orbit camera flag
+	// to allow for expected behaviors of the base class when in fly camera mode.
+	EUVEditor3DViewportClientCameraMode CameraMode = EUVEditor3DViewportClientCameraMode::Orbit;
+};
\ No newline at end of file
diff --git a/Engine/Plugins/Experimental/UVEditor/Source/UVEditor/Public/UVEditorCommands.h b/Engine/Plugins/Experimental/UVEditor/Source/UVEditor/Public/UVEditorCommands.h
--- a/Engine/Plugins/Experimental/UVEditor/Source/UVEditor/Public/UVEditorCommands.h
+++ b/Engine/Plugins/Experimental/UVEditor/Source/UVEditor/Public/UVEditorCommands.h
@@ -25,4 +25,7 @@ public:
 
 	TSharedPtr<FUICommandInfo> AcceptActiveTool;
 	TSharedPtr<FUICommandInfo> CancelActiveTool;
+
+	TSharedPtr<FUICommandInfo> EnableOrbitCamera;
+	TSharedPtr<FUICommandInfo> EnableFlyCamera;
 };
\ No newline at end of file
