UVEditor: Fix issue with non-unity builds related to EditorStyleSet not being available for UVEditor3DViewportToolbar

#rnx

diff --git a/Engine/Plugins/Experimental/UVEditor/Source/UVEditor/Private/SUVEditor3DViewportToolBar.cpp b/Engine/Plugins/Experimental/UVEditor/Source/UVEditor/Private/SUVEditor3DViewportToolBar.cpp
--- a/Engine/Plugins/Experimental/UVEditor/Source/UVEditor/Private/SUVEditor3DViewportToolBar.cpp
+++ b/Engine/Plugins/Experimental/UVEditor/Source/UVEditor/Private/SUVEditor3DViewportToolBar.cpp
@@ -8,7 +8,7 @@
 #include "Framework/MultiBox/MultiBoxBuilder.h"
 #include "Widgets/Layout/SBorder.h"
 #include "Widgets/SBoxPanel.h"
-
+#include "EditorStyleSet.h"
 
 void SUVEditor3DViewportToolBar::Construct(const FArguments& InArgs)
 {
