UVEditor: Conformal Island Unwrap action now added to the select tool, allowing per island conformal relaxations. Additionally, a bug was fixed with applying changes from the UV Editor when the channel number had been changed, and the sew tool was adjusted to prevent an error with the selection.

#rnx

diff --git a/Engine/Plugins/Experimental/UVEditor/Source/UVEditorTools/Private/UVIslandConformalUnwrapAction.cpp b/Engine/Plugins/Experimental/UVEditor/Source/UVEditorTools/Private/UVIslandConformalUnwrapAction.cpp
new file mode 100644
--- /dev/null
+++ b/Engine/Plugins/Experimental/UVEditor/Source/UVEditorTools/Private/UVIslandConformalUnwrapAction.cpp
@@ -0,0 +1,176 @@
+// Copyright Epic Games, Inc. All Rights Reserved.
+
+#include "UVIslandConformalUnwrapAction.h"
+
+#include "Algo/Unique.h"
+#include "ToolTargets/UVEditorToolMeshInput.h"
+#include "Selection/DynamicMeshSelection.h"
+#include "MeshOpPreviewHelpers.h" // UMeshOpPreviewWithBackgroundCompute
+#include "PreviewMesh.h"
+#include "ToolSetupUtil.h"
+#include "UVToolContextObjects.h"
+#include "Parameterization/DynamicMeshUVEditor.h"
+#include "Selections/MeshConnectedComponents.h"
+
+#define LOCTEXT_NAMESPACE "UUVIslandConformalUnwrapAction"
+
+using namespace UE::Geometry;
+
+
+UUVIslandConformalUnwrapAction::UUVIslandConformalUnwrapAction()
+{
+	CurrentSelection = MakeShared<UE::Geometry::FDynamicMeshSelection>();
+}
+
+void UUVIslandConformalUnwrapAction::SetWorld(UWorld* WorldIn)
+{
+	UUVToolAction::SetWorld(WorldIn);
+}
+
+void UUVIslandConformalUnwrapAction::Setup(UInteractiveTool* ParentToolIn)
+{
+	UUVToolAction::Setup(ParentToolIn);
+}
+
+void UUVIslandConformalUnwrapAction::Shutdown()
+{
+
+}
+
+
+void UUVIslandConformalUnwrapAction::SetSelection(int32 SelectionTargetIndexIn, const UE::Geometry::FDynamicMeshSelection* NewSelection)
+{
+	SelectionTargetIndex = SelectionTargetIndexIn;
+	if (NewSelection)
+	{
+		(*CurrentSelection) = *NewSelection;
+	}
+	else
+	{
+		(*CurrentSelection) = UE::Geometry::FDynamicMeshSelection();
+	}
+
+	UpdateVisualizations();
+}
+
+bool UUVIslandConformalUnwrapAction::PreCheckAction() 
+{
+	if (!CurrentSelection || CurrentSelection->Type != FDynamicMeshSelection::EType::Triangle)
+	{
+		return false;
+	}
+
+	if (!GatherIslandTids())
+	{
+		return false;
+	}
+
+	// Todo: Add some check here to confirm we have islands and not random triangles... like completely selected connected components?
+
+	return true;
+}
+
+bool UUVIslandConformalUnwrapAction::GatherIslandTids() {
+	IslandStartIndices.Reset();
+	ConcatenatedIslandTids.Reset();
+	MaxIslandSize = 0;
+
+	FDynamicMesh3& AppliedMesh = *(Targets[SelectionTargetIndex]->AppliedCanonical);
+	FDynamicMeshUVOverlay* UseOverlay = AppliedMesh.Attributes()->GetUVLayer(Targets[SelectionTargetIndex]->UVLayerIndex);
+
+	FMeshConnectedComponents ConnectedComponents(&AppliedMesh);
+
+	ConnectedComponents.FindConnectedTriangles([&](int32 Triangle0, int32 Triangle1) {
+		return UseOverlay->AreTrianglesConnected(Triangle0, Triangle1);
+	});
+	
+	int32 NumComponents = ConnectedComponents.Num();
+	IslandStartIndices.Add(0);
+	for (int32 k = 0; k < NumComponents; ++k)
+	{
+		TSet<int32> ComponentTris = TSet<int32>(ConnectedComponents[k].Indices);
+		TSet<int32> SelectedComponentTris = ComponentTris.Intersect(CurrentSelection->SelectedIDs);
+		if (SelectedComponentTris.Num() == ComponentTris.Num())
+		{
+			ConcatenatedIslandTids.Append(SelectedComponentTris.Array());
+			IslandStartIndices.Add(IslandStartIndices.Last(0) + SelectedComponentTris.Num());
+			MaxIslandSize = FMath::Max(MaxIslandSize, SelectedComponentTris.Num());
+		}
+	}
+
+	return true;
+}
+
+bool UUVIslandConformalUnwrapAction::ApplyAction(UUVToolEmitChangeAPI& EmitChangeAPI)
+{
+	FDynamicMesh3& MeshToUnwrap = *(Targets[SelectionTargetIndex]->AppliedCanonical);
+	FDynamicMeshUVOverlay* UVOverlay = MeshToUnwrap.Attributes()->GetUVLayer(Targets[SelectionTargetIndex]->UVLayerIndex);
+	FDynamicMeshUVEditor UVEditor(&MeshToUnwrap, Targets[SelectionTargetIndex]->UVLayerIndex, true);
+
+	FDynamicMeshChangeTracker ChangeTracker(Targets[SelectionTargetIndex]->UnwrapCanonical.Get());
+	ChangeTracker.BeginChange();
+	ChangeTracker.SaveTriangles(ConcatenatedIslandTids, true);
+
+	TArray<int32> ComponentTris;
+	ComponentTris.Reserve(MaxIslandSize);
+	int32 NumComponents = IslandStartIndices.Num()-1;
+	TArray<bool> bComponentSolved;
+	bComponentSolved.Init(false, NumComponents);
+	int32 SuccessCount = 0;
+	TSet<int32> VidSet;
+	for (int32 k = 0; k < NumComponents; ++k)
+	{
+
+
+		ComponentTris.Reset(0);
+		int32 StartIndex = IslandStartIndices[k];
+		int32 Count = IslandStartIndices[k+1] - StartIndex;
+		if (Count == 0)
+		{
+			continue; // Skip everything if there are no Tids for this island component
+		}
+		ComponentTris.Append(&ConcatenatedIslandTids[StartIndex], Count);
+
+
+		FAxisAlignedBox2f UVBounds = FAxisAlignedBox2f::Empty();
+		for (int32 tid : ComponentTris)
+		{
+			if (UVOverlay->IsSetTriangle(tid))
+			{
+				FIndex3i UVTri = UVOverlay->GetTriangle(tid);
+				FVector2f U = UVOverlay->GetElement(UVTri.A);
+				FVector2f V = UVOverlay->GetElement(UVTri.B);
+				FVector2f W = UVOverlay->GetElement(UVTri.C);
+				UVBounds.Contain(U); UVBounds.Contain(V); UVBounds.Contain(W);
+			}
+		}
+
+
+		bComponentSolved[k] = UVEditor.SetTriangleUVsFromFreeBoundaryConformal(ComponentTris, true);
+		if (bComponentSolved[k])
+		{
+			UVEditor.ScaleUVAreaToBoundingBox(ComponentTris, UVBounds, true, true);
+		}
+	
+		for (int32 Tid : ComponentTris)
+		{
+			FIndex3i TriVids = MeshToUnwrap.Attributes()->GetUVLayer(Targets[SelectionTargetIndex]->UVLayerIndex)->GetTriangle(Tid);
+			for (int i = 0; i < 3; ++i)
+			{
+				VidSet.Add(TriVids[i]);
+			}
+		}
+	}
+
+	const TArray<int32> AllVids = VidSet.Array();
+
+	Targets[SelectionTargetIndex]->UpdateAllFromAppliedCanonical(&AllVids, UUVEditorToolMeshInput::NONE_CHANGED_ARG, &ConcatenatedIslandTids);
+	
+	const FText TransactionName(LOCTEXT("ConformalUnwrapCompleteTransactionName", "Conformal Unwrap Islands"));
+	EmitChangeAPI.EmitToolIndependentUnwrapCanonicalChange(Targets[SelectionTargetIndex],
+		ChangeTracker.EndChange(), TransactionName);
+
+	return true;
+}
+
+#undef LOCTEXT_NAMESPACE
\ No newline at end of file
diff --git a/Engine/Plugins/Experimental/UVEditor/Source/UVEditorTools/Private/UVSeamSewAction.cpp b/Engine/Plugins/Experimental/UVEditor/Source/UVEditorTools/Private/UVSeamSewAction.cpp
--- a/Engine/Plugins/Experimental/UVEditor/Source/UVEditorTools/Private/UVSeamSewAction.cpp
+++ b/Engine/Plugins/Experimental/UVEditor/Source/UVEditorTools/Private/UVSeamSewAction.cpp
@@ -19,7 +19,7 @@ using namespace UE::Geometry;
 
 UUVSeamSewAction::UUVSeamSewAction()
 {
-
+	CurrentSelection = MakeShared<UE::Geometry::FDynamicMeshSelection>();
 }
 
 void UUVSeamSewAction::SetWorld(UWorld* WorldIn)
@@ -72,7 +72,15 @@ void UUVSeamSewAction::Shutdown()
 void UUVSeamSewAction::SetSelection(int32 SelectionTargetIndexIn, const UE::Geometry::FDynamicMeshSelection* NewSelection)
 {
 	SelectionTargetIndex = SelectionTargetIndexIn;
-	CurrentSelection = NewSelection;
+	if (NewSelection)
+	{
+		(*CurrentSelection) = *NewSelection;
+	}
+	else
+	{
+		(*CurrentSelection) = UE::Geometry::FDynamicMeshSelection();
+	}
+
 	EdgeSewCandidates.Reset();
 
 	if (SelectionTargetIndex != -1 && CurrentSelection && CurrentSelection->Type == FDynamicMeshSelection::EType::Edge)
@@ -99,31 +107,15 @@ void UUVSeamSewAction::SetSelection(int32 SelectionTargetIndexIn, const UE::Geom
 	UpdateVisualizations();
 }
 
-
-void UUVSeamSewAction::Tick(float DeltaTime) 
-{
-
-
-}
-
 void UUVSeamSewAction::UpdateVisualizations()
 {
 	UpdateSewEdgePreviewLines();
 }
 
-bool UUVSeamSewAction::ExecuteAction(UUVToolEmitChangeAPI& EmitChangeAPI)
-{
-	if (!PreCheckSewAction()) {
-		return false;
-	}
-
-	return ApplySewAction(EmitChangeAPI);
-}
-
 void UUVSeamSewAction::UpdateSewEdgePreviewLines()
 {
 	SewEdgePairingLineSet->Clear();
-	if (CurrentSelection != nullptr && !CurrentSelection->IsEmpty())
+	if (CurrentSelection && !CurrentSelection->IsEmpty())
 	{
 		FTransform MeshTransform = Targets[SelectionTargetIndex]->UnwrapPreview->PreviewMesh->GetTransform();
 
@@ -153,9 +145,9 @@ void UUVSeamSewAction::UpdateSewEdgePreviewLines()
 }
 
 
-bool UUVSeamSewAction::PreCheckSewAction() const
+bool UUVSeamSewAction::PreCheckAction()
 {
-	if (SelectionTargetIndex == -1 || CurrentSelection == nullptr || CurrentSelection->Mesh == nullptr)
+	if (SelectionTargetIndex == -1 || !CurrentSelection || CurrentSelection->Mesh == nullptr)
 	{
 		UE_LOG(LogGeometry, Warning, TEXT("Cannot sew UVs. Mesh selection was empty."));
 		return false;
@@ -227,7 +219,7 @@ int32 UUVSeamSewAction::FindSewEdgeOppositePairing(int32 UnwrapEid) const
 }
 
 
-bool UUVSeamSewAction::ApplySewAction(UUVToolEmitChangeAPI& EmitChangeAPI)
+bool UUVSeamSewAction::ApplyAction(UUVToolEmitChangeAPI& EmitChangeAPI)
 {
 	FDynamicMesh3& MeshToSew = *(Targets[SelectionTargetIndex]->UnwrapCanonical);
 
@@ -293,10 +285,8 @@ bool UUVSeamSewAction::ApplySewAction(UUVToolEmitChangeAPI& EmitChangeAPI)
 	checkSlow(MeshToSew.IsSameAs(*Targets[SelectionTargetIndex]->UnwrapPreview->PreviewMesh->GetMesh(), FDynamicMesh3::FSameAsOptions()));
 
 	const FText TransactionName(LOCTEXT("SewCompleteTransactionName", "Sew Edges"));
-	EmitChangeAPI.BeginUndoTransaction(TransactionName);
 	EmitChangeAPI.EmitToolIndependentUnwrapCanonicalChange(Targets[SelectionTargetIndex],
 		ChangeTracker.EndChange(), TransactionName);
-	EmitChangeAPI.EndUndoTransaction();
 
 	SetSelection(-1, nullptr);
 
diff --git a/Engine/Plugins/Experimental/UVEditor/Source/UVEditorTools/Private/UVSelectTool.cpp b/Engine/Plugins/Experimental/UVEditor/Source/UVEditorTools/Private/UVSelectTool.cpp
--- a/Engine/Plugins/Experimental/UVEditor/Source/UVEditorTools/Private/UVSelectTool.cpp
+++ b/Engine/Plugins/Experimental/UVEditor/Source/UVEditorTools/Private/UVSelectTool.cpp
@@ -20,6 +20,7 @@
 #include "ToolSetupUtil.h"
 
 #include "UVSeamSewAction.h"
+#include "UVIslandConformalUnwrapAction.h"
 
 #include "ToolTargetManager.h"
 
@@ -206,6 +207,11 @@ void  USelectToolActionPropertySet::Sew()
 	PostAction(ESelectToolAction::Sew);
 }
 
+void  USelectToolActionPropertySet::IslandConformalUnwrap()
+{
+	PostAction(ESelectToolAction::IslandConformalUnwrap);
+}
+
 void USelectToolActionPropertySet::PostAction(ESelectToolAction Action)
 {
 	if (ParentTool.IsValid())
@@ -334,6 +340,10 @@ void UUVSelectTool::Setup()
 	SewAction->SetTargets(Targets);
 	SewAction->SetWorld(Targets[0]->UnwrapPreview->GetWorld());
 
+	IslandConformalUnwrapAction = NewObject<UUVIslandConformalUnwrapAction>();
+	IslandConformalUnwrapAction->Setup(this);
+	IslandConformalUnwrapAction->SetTargets(Targets);
+	IslandConformalUnwrapAction->SetWorld(Targets[0]->UnwrapPreview->GetWorld());
 
 	if (!SelectionMechanic->GetCurrentSelection().IsEmpty())
 	{
@@ -373,6 +383,11 @@ void UUVSelectTool::Shutdown(EToolShutdownType ShutdownType)
 		SewAction->Shutdown();
 	}
 
+	if (IslandConformalUnwrapAction)
+	{
+		IslandConformalUnwrapAction->Shutdown();
+	}
+
 	// Calls shutdown on gizmo and destroys it.
 	GetToolManager()->GetPairedGizmoManager()->DestroyAllGizmosByOwner(this);
 
@@ -558,6 +573,7 @@ void UUVSelectTool::OnSelectionChanged()
 	}
 
 	SewAction->SetSelection(SelectionTargetIndex, &Selection);
+	IslandConformalUnwrapAction->SetSelection(SelectionTargetIndex, &Selection);
 
 	UpdateLivePreviewLines();
 	UpdateGizmo();
@@ -697,6 +713,7 @@ void UUVSelectTool::ApplyGizmoTransform()
 
 		bGizmoTransformNeedsApplication = false;
 		SewAction->UpdateVisualizations();
+		IslandConformalUnwrapAction->UpdateVisualizations();
 	}
 }
 
@@ -737,20 +754,25 @@ void UUVSelectTool::ApplyAction(ESelectToolAction ActionType)
 	case ESelectToolAction::Sew:
 		if (SewAction)
 		{
+			const FText TransactionName(LOCTEXT("SewCompleteTransactionName", "Sew Edges"));
+			EmitChangeAPI->BeginUndoTransaction(TransactionName);
+
+			SelectionMechanic->SetSelection(FDynamicMeshSelection(), false, true);
 			bool ActionSuccessful = SewAction->ExecuteAction(*EmitChangeAPI);
 
-			if (ActionSuccessful)
-			{
-				if (!AABBTrees[SelectionTargetIndex]->IsValid(false))
-				{
-					AABBTrees[SelectionTargetIndex]->Build();
-				}
+			EmitChangeAPI->EndUndoTransaction();
+		}
+		break;
+	case ESelectToolAction::IslandConformalUnwrap:
+		if (IslandConformalUnwrapAction)
+		{
+			const FText TransactionName(LOCTEXT("ConformalUnwrapCompleteTransactionName", "Conformal Unwrap Islands"));
+			EmitChangeAPI->BeginUndoTransaction(TransactionName);
 
-				// TODO: Emit a selection update instead of just clearing things here. If we do
-				// still clear, remember that SelectionTargetIndex gets updated, so either do it
-				// here or save a pointer to input object before doing stuff.
-				SelectionMechanic->SetSelection(FDynamicMeshSelection(), true, false);
-			}
+			SelectionMechanic->SetSelection(FDynamicMeshSelection(), false, true);
+			bool ActionSuccessful = IslandConformalUnwrapAction->ExecuteAction(*EmitChangeAPI);
+
+			EmitChangeAPI->EndUndoTransaction();
 		}
 		break;
 	default:
diff --git a/Engine/Plugins/Experimental/UVEditor/Source/UVEditorTools/Public/UVIslandConformalUnwrapAction.h b/Engine/Plugins/Experimental/UVEditor/Source/UVEditorTools/Public/UVIslandConformalUnwrapAction.h
new file mode 100644
--- /dev/null
+++ b/Engine/Plugins/Experimental/UVEditor/Source/UVEditorTools/Public/UVIslandConformalUnwrapAction.h
@@ -0,0 +1,41 @@
+// Copyright Epic Games, Inc. All Rights Reserved.
+
+#pragma once
+
+#include "CoreMinimal.h"
+#include "GeometryBase.h"
+#include "IndexTypes.h"
+
+#include "UVToolAction.h"
+
+#include "UVIslandConformalUnwrapAction.generated.h"
+
+PREDECLARE_GEOMETRY(class FDynamicMeshSelection);
+
+UCLASS()
+class UVEDITORTOOLS_API UUVIslandConformalUnwrapAction : public UUVToolAction
+{
+	GENERATED_BODY()
+
+public:
+
+	UUVIslandConformalUnwrapAction();
+
+	void SetWorld(UWorld* WorldIn) override;
+	void Setup(UInteractiveTool* ParentToolIn) override;
+	void Shutdown() override;
+	void SetSelection(int32 SelectionTargetIndexIn, const UE::Geometry::FDynamicMeshSelection* NewSelection);
+
+protected:
+
+	TArray<int32> ConcatenatedIslandTids;
+	TArray<int32> IslandStartIndices;
+	int32 MaxIslandSize;
+
+	int32 SelectionTargetIndex;
+	TSharedPtr<UE::Geometry::FDynamicMeshSelection> CurrentSelection;
+
+	bool GatherIslandTids();
+	bool PreCheckAction();
+	bool ApplyAction(UUVToolEmitChangeAPI& EmitChangeAPI);	
+};
diff --git a/Engine/Plugins/Experimental/UVEditor/Source/UVEditorTools/Public/UVSeamSewAction.h b/Engine/Plugins/Experimental/UVEditor/Source/UVEditorTools/Public/UVSeamSewAction.h
--- a/Engine/Plugins/Experimental/UVEditor/Source/UVEditorTools/Public/UVSeamSewAction.h
+++ b/Engine/Plugins/Experimental/UVEditor/Source/UVEditorTools/Public/UVSeamSewAction.h
@@ -46,8 +46,6 @@ public:
 	void Shutdown() override;
 
 	void SetSelection(int32 SelectionTargetIndexIn, const UE::Geometry::FDynamicMeshSelection* NewSelection);
-	virtual void Tick(float DeltaTime) override;
-	virtual bool ExecuteAction(UUVToolEmitChangeAPI& EmitChangeAPI) override;
 	virtual void UpdateVisualizations() override;
 
 protected:
@@ -61,10 +59,10 @@ protected:
 	TObjectPtr<ULineSetComponent> SewEdgePairingLineSet = nullptr;
 
 	int32 SelectionTargetIndex;
-	const UE::Geometry::FDynamicMeshSelection* CurrentSelection;
+	TSharedPtr<UE::Geometry::FDynamicMeshSelection> CurrentSelection;
 
 	int32 FindSewEdgeOppositePairing(int32 UnwrapEid) const;
-	bool PreCheckSewAction() const;
-	bool ApplySewAction(UUVToolEmitChangeAPI& EmitChangeAPI);
+	bool PreCheckAction() override;
+	bool ApplyAction(UUVToolEmitChangeAPI& EmitChangeAPI) override;
 	void UpdateSewEdgePreviewLines();
 };
\ No newline at end of file
diff --git a/Engine/Plugins/Experimental/UVEditor/Source/UVEditorTools/Public/UVSelectTool.h b/Engine/Plugins/Experimental/UVEditor/Source/UVEditorTools/Public/UVSelectTool.h
--- a/Engine/Plugins/Experimental/UVEditor/Source/UVEditorTools/Public/UVSelectTool.h
+++ b/Engine/Plugins/Experimental/UVEditor/Source/UVEditorTools/Public/UVSelectTool.h
@@ -29,6 +29,7 @@ class UUVToolStateObjectStore;
 class UPreviewMesh;
 class UUVEditorToolMeshInput;
 class UUVSeamSewAction;
+class UUVIslandConformalUnwrapAction;
 class UUVToolEmitChangeAPI;
 class UUVToolViewportButtonsAPI;
 
@@ -82,6 +83,7 @@ enum class ESelectToolAction
 	NoAction,
 
 	Sew,
+	IslandConformalUnwrap
 };
 
 UCLASS()
@@ -95,10 +97,14 @@ public:
 	void Initialize(UUVSelectTool* ParentToolIn) { ParentTool = ParentToolIn; }
 	void PostAction(ESelectToolAction Action);
 
-	/** Move the mirror plane to center of bounding box without changing its normal. */
+	/* Sew selected boundary edges together topologically in the UV set */
 	UFUNCTION(CallInEditor, Category = Actions)
 	void Sew();
 
+	/* Apply a conformal unwrap to the selected UV islands */
+	UFUNCTION(CallInEditor, Category = Actions)
+	void IslandConformalUnwrap();
+
 };
 
 
@@ -197,6 +203,9 @@ protected:
 	UPROPERTY()
 	TObjectPtr<UUVSeamSewAction> SewAction = nullptr;
 
+	UPROPERTY()
+	TObjectPtr<UUVIslandConformalUnwrapAction> IslandConformalUnwrapAction = nullptr;
+
 	UE::Geometry::FFrame3d InitialGizmoFrame;
 	FTransform UnappliedGizmoTransform;
 	bool bInDrag = false;
diff --git a/Engine/Plugins/Experimental/UVEditor/Source/UVEditorTools/Public/UVToolAction.h b/Engine/Plugins/Experimental/UVEditor/Source/UVEditorTools/Public/UVToolAction.h
--- a/Engine/Plugins/Experimental/UVEditor/Source/UVEditorTools/Public/UVToolAction.h
+++ b/Engine/Plugins/Experimental/UVEditor/Source/UVEditorTools/Public/UVToolAction.h
@@ -33,11 +33,20 @@ public:
 			Target = nullptr;
 		}
 	}
-	virtual bool ExecuteAction(UUVToolEmitChangeAPI& EmitChangeAPI) { return true; };
+	virtual bool ExecuteAction(UUVToolEmitChangeAPI& EmitChangeAPI) { 
+		if (!PreCheckAction()) {
+			return false;
+		}
+
+		return ApplyAction(EmitChangeAPI);
+	};
 	virtual void UpdateVisualizations() {};
 
 protected:
 
+	virtual bool PreCheckAction() { return true; };
+	virtual bool ApplyAction(UUVToolEmitChangeAPI& EmitChangeAPI) { return true; };
+
 	UPROPERTY()
 	TArray<TObjectPtr<UUVEditorToolMeshInput>> Targets;
 };
diff --git a/Engine/Plugins/Runtime/GeometryProcessing/Source/DynamicMesh/Private/Parameterization/DynamicMeshUVEditor.cpp b/Engine/Plugins/Runtime/GeometryProcessing/Source/DynamicMesh/Private/Parameterization/DynamicMeshUVEditor.cpp
--- a/Engine/Plugins/Runtime/GeometryProcessing/Source/DynamicMesh/Private/Parameterization/DynamicMeshUVEditor.cpp
+++ b/Engine/Plugins/Runtime/GeometryProcessing/Source/DynamicMesh/Private/Parameterization/DynamicMeshUVEditor.cpp
@@ -1156,22 +1156,18 @@ bool FDynamicMeshUVEditor::ScaleUVAreaTo3DArea(const TArray<int32>& Triangles, b
 	}
 
 	TSet<int32> Elements;
+	FAxisAlignedBox2f UVBounds;
+	double Area2D = DetermineAreaFromUVs(*UVOverlay, Triangles, &UVBounds);
 
-	double Area2D = 0.0;
-	FAxisAlignedBox2f UVBounds = FAxisAlignedBox2f::Empty();
 	for (int32 tid : Triangles)
 	{
 		if (UVOverlay->IsSetTriangle(tid))
 		{
 			FIndex3i UVTri = UVOverlay->GetTriangle(tid);
 			Elements.Add(UVTri.A); Elements.Add(UVTri.B); Elements.Add(UVTri.C);
-			FVector2f U = UVOverlay->GetElement(UVTri.A);
-			FVector2f V = UVOverlay->GetElement(UVTri.B);
-			FVector2f W = UVOverlay->GetElement(UVTri.C);
-			UVBounds.Contain(U); UVBounds.Contain(V); UVBounds.Contain(W);
-			Area2D += (double)VectorUtil::Area(U, V, W);
 		}
 	}
+
 	if (Elements.Num() == 0 || FMathd::Abs(Area2D) < FMathf::Epsilon ||  FMathd::IsFinite(Area2D) == false )
 	{
 		return false;
@@ -1195,6 +1191,61 @@ bool FDynamicMeshUVEditor::ScaleUVAreaTo3DArea(const TArray<int32>& Triangles, b
 	return true;
 }
 
+bool FDynamicMeshUVEditor::ScaleUVAreaToBoundingBox(const TArray<int32>& Triangles, const FAxisAlignedBox2f& BoundingBox, bool bPreserveAspectRatio, bool bRecenterAtBoundingBox)
+{
+	if (FMathd::Abs(BoundingBox.Area()) < FMathf::Epsilon || FMathd::IsFinite(BoundingBox.Area()) == false)
+	{
+		return false;
+	}
+
+	TSet<int32> Elements;
+	FAxisAlignedBox2f UVBounds;
+	double Area2D = DetermineAreaFromUVs(*UVOverlay, Triangles, &UVBounds);
+
+	for (int32 tid : Triangles)
+	{
+		if (UVOverlay->IsSetTriangle(tid))
+		{
+			FIndex3i UVTri = UVOverlay->GetTriangle(tid);
+			Elements.Add(UVTri.A); Elements.Add(UVTri.B); Elements.Add(UVTri.C);
+		}
+	}
+	
+
+	if (Elements.Num() == 0 || FMathd::Abs(Area2D) < FMathf::Epsilon || FMathd::IsFinite(Area2D) == false)
+	{
+		return false;
+	}
+
+	float WidthScale = BoundingBox.Width() / UVBounds.Width();
+	float HeightScale = BoundingBox.Height() / UVBounds.Height();
+
+	if (bPreserveAspectRatio)
+	{
+		WidthScale = FMath::Min(WidthScale, HeightScale);
+		HeightScale = WidthScale;
+	}
+
+	if ( !FMathd::IsFinite(WidthScale) || !FMathd::IsFinite(HeightScale))
+	{
+		return false;
+	}
+
+	FVector2f ScaleOrigin = UVBounds.Center();
+	FVector2f Translation = (bRecenterAtBoundingBox) ? BoundingBox.Center() : ScaleOrigin;
+	TransformUVElements(Elements.Array(), [ScaleOrigin, WidthScale, HeightScale, Translation](const FVector2f& UV){
+		FVector2f UVTransformed = UV;
+		UVTransformed = (UVTransformed - ScaleOrigin);
+		UVTransformed[0] = UVTransformed[0] * WidthScale;
+		UVTransformed[1] = UVTransformed[1] * HeightScale;
+		UVTransformed = UVTransformed + Translation;
+		return UVTransformed;
+	});
+
+	return true;
+}
+
+
 
 
 bool FDynamicMeshUVEditor::AutoOrientUVArea(const TArray<int32>& Triangles)
@@ -1319,4 +1370,30 @@ bool FDynamicMeshUVEditor::QuickPack(int32 TargetTextureResolution, float Gutter
 	bool bOK = Packer.StandardPack();
 
 	return bOK;
-}
\ No newline at end of file
+}
+
+double  FDynamicMeshUVEditor::DetermineAreaFromUVs(const FDynamicMeshUVOverlay& UVOverlay, const TArray<int32>& Triangles, FAxisAlignedBox2f* BoundingBox)
+{
+	if (BoundingBox)
+	{
+		(*BoundingBox) = FAxisAlignedBox2f::Empty();
+	}
+	double Area2D = 0.0;
+	for (int32 tid : Triangles)
+	{
+		if (UVOverlay.IsSetTriangle(tid))
+		{
+			FIndex3i UVTri = UVOverlay.GetTriangle(tid);			
+			FVector2f U = UVOverlay.GetElement(UVTri.A);
+			FVector2f V = UVOverlay.GetElement(UVTri.B);
+			FVector2f W = UVOverlay.GetElement(UVTri.C);
+			if (BoundingBox)
+			{
+				BoundingBox->Contain(U); BoundingBox->Contain(V); BoundingBox->Contain(W);
+			}
+			Area2D += (double)VectorUtil::Area(U, V, W);
+		}
+	}
+
+	return Area2D;
+}
diff --git a/Engine/Plugins/Runtime/GeometryProcessing/Source/DynamicMesh/Public/Parameterization/DynamicMeshUVEditor.h b/Engine/Plugins/Runtime/GeometryProcessing/Source/DynamicMesh/Public/Parameterization/DynamicMeshUVEditor.h
--- a/Engine/Plugins/Runtime/GeometryProcessing/Source/DynamicMesh/Public/Parameterization/DynamicMeshUVEditor.h
+++ b/Engine/Plugins/Runtime/GeometryProcessing/Source/DynamicMesh/Public/Parameterization/DynamicMeshUVEditor.h
@@ -205,6 +205,16 @@ public:
 	*/
 	bool ScaleUVAreaTo3DArea(const TArray<int32>& Triangles, bool bRecenterAtOrigin);
 
+
+	/**
+	* Scale UVs of the given triangles to fit within the provided provided bounding box.
+	* @param BoundingBox the 2D bounding box to constraint the requested UVs within.
+	* @param bPreserveAspectRatio if true, when scaling UVs, maintain original aspect ratio instead of fitting to box bounds
+	* @param bRecenterAtBoundingBox if true, UVs are translated such that their centroid is at the bounding box center
+	*/
+	bool ScaleUVAreaToBoundingBox(const TArray<int32>& Triangles, const FAxisAlignedBox2f& BoundingBox, bool bPreserveAspectRatio, bool bRecenterAtBoundingBox);
+
+
 	/**
 	* Compute an oriented UV-space bounding box for the given Triangles, and then rotate the UVs such that the
 	* box is aligned with the X axis. 
@@ -233,6 +243,11 @@ public:
 	 */
 	static bool EstimateGeodesicCenterFrameVertex(const FDynamicMesh3& Mesh, const TArray<int32>& Triangles, FFrame3d& FrameOut, int32& VertexIDOut, bool bAlignToUnitAxes = true);
 
+	/**
+	* Compute the occupied 2D area for the UVs for the specified triangles and channel.
+	* @param BoundingBox if not null, return the overall bounding box of the UVs for the specified triangles along with the exact area occupied.
+	*/
+	static double DetermineAreaFromUVs(const FDynamicMeshUVOverlay& UVOverlay, const TArray<int32>& Triangles, FAxisAlignedBox2f* BoundingBox = nullptr);
 };
 
 
