UVEditor: Fix tool tab not being reopenable after close, and crash when using a layout without viewport.

#rnx

diff --git a/Engine/Plugins/Experimental/UVEditor/Source/UVEditor/Private/UVEditorToolkit.cpp b/Engine/Plugins/Experimental/UVEditor/Source/UVEditor/Private/UVEditorToolkit.cpp
--- a/Engine/Plugins/Experimental/UVEditor/Source/UVEditor/Private/UVEditorToolkit.cpp
+++ b/Engine/Plugins/Experimental/UVEditor/Source/UVEditor/Private/UVEditorToolkit.cpp
@@ -182,18 +182,18 @@ void FUVEditorToolkit::RegisterTabSpawners(const TSharedRef<FTabManager>& InTabM
 
 	InTabManager->RegisterTabSpawner(InteractiveToolsPanelTabID, FOnSpawnTab::CreateSP(this, 
 		&FUVEditorToolkit::SpawnTab_InteractiveToolsPanel))
-		.SetDisplayName(LOCTEXT("InteractiveToolsPanel", "Tools Panel"))
+		.SetDisplayName(LOCTEXT("UVToolsPanelLabel", "UV Tools Panel"))
 		.SetGroup(AssetEditorTabsCategory.ToSharedRef())
 		.SetIcon(FSlateIcon(FEditorStyle::GetStyleSetName(), "LevelEditor.Tabs.Details"));
 
 	InTabManager->RegisterTabSpawner(ViewportTabID, FOnSpawnTab::CreateSP(this, &FUVEditorToolkit::SpawnTab_Viewport))
-		.SetDisplayName(LOCTEXT("ViewportTabUV", "UV Viewport"))
+		.SetDisplayName(LOCTEXT("2DViewportTabLabel", "2D Viewport"))
 		.SetGroup(AssetEditorTabsCategory.ToSharedRef())
 		.SetIcon(FSlateIcon(FEditorStyle::GetStyleSetName(), "LevelEditor.Tabs.Viewports"));
 
 	InTabManager->RegisterTabSpawner(LivePreviewTabID, FOnSpawnTab::CreateSP(this, 
 		&FUVEditorToolkit::SpawnTab_LivePreview))
-		.SetDisplayName(LOCTEXT("ViewportTabLivePreview", "3D Viewport"))
+		.SetDisplayName(LOCTEXT("3DViewportTabLabel", "3D Viewport"))
 		.SetGroup(AssetEditorTabsCategory.ToSharedRef())
 		.SetIcon(FSlateIcon(FEditorStyle::GetStyleSetName(), "LevelEditor.Tabs.Viewports"));
 }
@@ -275,10 +275,27 @@ TSharedRef<SDockTab> FUVEditorToolkit::SpawnTab_LivePreview(const FSpawnTabArgs&
 // This is bound in RegisterTabSpawners() to create the panel on the left. The panel is filled in by the mode.
 TSharedRef<SDockTab> FUVEditorToolkit::SpawnTab_InteractiveToolsPanel(const FSpawnTabArgs& Args)
 {
-	ToolsPanel = SNew(SDockTab)
-		.Label(LOCTEXT("UVToolPanelTitle", "UV Tools"));
+	TSharedRef<SDockTab> ToolsPanel = SNew(SDockTab);
 
-	return ToolsPanel.ToSharedRef();
+	UUVEditorMode* UVMode = Cast<UUVEditorMode>(EditorModeManager->GetActiveScriptableMode(UUVEditorMode::EM_UVEditorModeId));
+	if (!UVMode)
+	{
+		// This is where we will drop out on the first call to this callback, when the mode does not yet
+		// exist. There is probably a place where we could safely intialize the mode to make sure that
+		// it is around before the first call, but it seems cleaner to just do the mode initialization 
+		// in PostInitAssetEditor and fill in the tab at that time.
+		// Later calls to this callback will occur if the user closes and restores the tab, and they
+		// will continue past this point to allow the tab to be refilled.
+		return ToolsPanel;
+	}
+	TSharedPtr<FModeToolkit> UVModeToolkit = UVMode->GetToolkit().Pin();
+	if (!UVModeToolkit.IsValid())
+	{
+		return ToolsPanel;
+	}
+
+	ToolsPanel->SetContent(UVModeToolkit->GetInlineContent().ToSharedRef());
+	return ToolsPanel;
 }
 
 void FUVEditorToolkit::CreateWidgets()
@@ -368,10 +385,24 @@ void FUVEditorToolkit::PostInitAssetEditor()
 
 	UVMode->InitializeTargets(ObjectsToEdit, ObjectTransforms);
 
+	// Regardless of how the user has modified the layout, we're going to make sure that we have
+	// a tool panel to show the tools, and a 2d viewport that will allow our mode to receive ticks.
+	TSharedPtr<SDockTab> ToolsPanel = TabManager->FindExistingLiveTab(InteractiveToolsPanelTabID);
+	if (!ToolsPanel)
+	{
+		ToolsPanel = TabManager->TryInvokeTab(InteractiveToolsPanelTabID);
+	}
+	if (!TabManager->FindExistingLiveTab(ViewportTabID))
+	{
+		TabManager->TryInvokeTab(ViewportTabID);
+	}
+
 	// Plug in the mode tool panel
 	TSharedPtr<FModeToolkit> UVModeToolkit = UVMode->GetToolkit().Pin();
-	check(UVModeToolkit.IsValid());
-	ToolsPanel->SetContent(UVModeToolkit->GetInlineContent().ToSharedRef());
+	if (ensure(ToolsPanel && UVModeToolkit.IsValid()))
+	{
+		ToolsPanel->SetContent(UVModeToolkit->GetInlineContent().ToSharedRef());
+	}
 
 	// Add the "Apply Changes" button. It should actually be safe to do this almost
 	// any time, even before that toolbar's registration, but it's easier to put most
diff --git a/Engine/Plugins/Experimental/UVEditor/Source/UVEditor/Public/UVEditorToolkit.h b/Engine/Plugins/Experimental/UVEditor/Source/UVEditor/Public/UVEditorToolkit.h
--- a/Engine/Plugins/Experimental/UVEditor/Source/UVEditor/Public/UVEditorToolkit.h
+++ b/Engine/Plugins/Experimental/UVEditor/Source/UVEditor/Public/UVEditorToolkit.h
@@ -66,9 +66,6 @@ protected:
 	// FAssetEditorToolkit
 	virtual void PostInitAssetEditor() override;
 
-	/** Inline content area for the UV mode's content (gotten from UVEditorModeToolkit) */
-	TSharedPtr<SDockTab> ToolsPanel;
-
 	/** Scene in which the 2D unwrapped uv meshes live. */
 	TUniquePtr<FPreviewScene> UnwrapScene;
 
