UVEditor: Added split action to the select tool. Also fixed edge selection occasionally breaking due to changing eids.

#rnx

diff --git a/Engine/Plugins/Experimental/MeshModelingToolsetExp/Source/MeshModelingToolsExp/Private/SeamSculptTool.cpp b/Engine/Plugins/Experimental/MeshModelingToolsetExp/Source/MeshModelingToolsExp/Private/SeamSculptTool.cpp
--- a/Engine/Plugins/Experimental/MeshModelingToolsetExp/Source/MeshModelingToolsExp/Private/SeamSculptTool.cpp
+++ b/Engine/Plugins/Experimental/MeshModelingToolsetExp/Source/MeshModelingToolsExp/Private/SeamSculptTool.cpp
@@ -331,8 +331,18 @@ void USeamSculptTool::CreateSeamAlongPath()
 		TUniquePtr<FMeshChange> MeshEditChange = PreviewMesh->TrackedEditMesh([&](FDynamicMesh3& Mesh, FDynamicMeshChangeTracker& Tracker)
 		{
 			Tracker.SaveVertexOneRingTriangles(CurDrawPath, true);
+
+			TSet<int32> PathEids;
+			for (int32 i = 0; i < CurDrawPath.Num() - 1; ++i)
+			{
+				int32 Eid = Mesh.FindEdge(CurDrawPath[i], CurDrawPath[i + 1]);
+				if (ensure(Eid != IndexConstants::InvalidID))
+				{
+					PathEids.Add(Eid);
+				}
+			}
 			FDynamicMeshUVEditor UVEditor(&Mesh, 0, false);
-			UVEditor.CreateSeamAlongVertexPath(CurDrawPath);
+			UVEditor.CreateSeamsAtEdges(PathEids);
 		});
 
 		GetToolManager()->EmitObjectChange(PreviewMesh, MoveTemp(MeshEditChange), LOCTEXT("CreateSeamChange", "Add UV Seam"));
diff --git a/Engine/Plugins/Experimental/UVEditor/Source/UVEditorTools/Private/UVSelectTool.cpp b/Engine/Plugins/Experimental/UVEditor/Source/UVEditorTools/Private/UVSelectTool.cpp
--- a/Engine/Plugins/Experimental/UVEditor/Source/UVEditorTools/Private/UVSelectTool.cpp
+++ b/Engine/Plugins/Experimental/UVEditor/Source/UVEditorTools/Private/UVSelectTool.cpp
@@ -30,6 +30,66 @@ using namespace UE::Geometry;
 
 namespace UVSelectToolLocals
 {
+	// These following three functions deal with the unfortunate problem that eids are unstable as identifiers
+	// (e.g. removing and reinserting the same triangles can change the eids of the edges), so edges have to
+	// be identified in another way. We identify them by vertex ID pairs. This should really be dealt with
+	// on a mesh selection level, but for now we fix it here.
+	// After selection changes, we convert our eids to vid pairs. After mesh changes, we update the selection
+	// eids from our stored vid pairs.
+
+	/** If selection is a non-empty edge selection, update its eids using stored vid pairs. */
+	void UpdateSelectionEidsAfterMeshChange(FDynamicMeshSelection& SelectionInOut, TArray<FIndex2i>* VidPairsIn)
+	{
+		if (!SelectionInOut.Mesh || SelectionInOut.Type != FDynamicMeshSelection::EType::Edge)
+		{
+			// No update necessary
+			return;
+		}
+
+		// Otherwise, updating eids.
+		if (!ensure(VidPairsIn))
+		{
+			return;
+		}
+		SelectionInOut.SelectedIDs.Empty();
+		for (const FIndex2i& VidPair : *VidPairsIn)
+		{
+			int32 Eid = SelectionInOut.Mesh->FindEdge(VidPair.A, VidPair.B);
+			if (ensure(Eid != IndexConstants::InvalidID))
+			{
+				SelectionInOut.SelectedIDs.Add(Eid);
+			}
+		}
+	}
+
+	/** If selection mechanic holds a non-empty edge selection, update its eids using stored vid pairs. */
+	void UpdateSelectionEidsAfterMeshChange(UMeshSelectionMechanic& SelectionMechanic, TArray<FIndex2i>* VidPairsIn)
+	{
+		const FDynamicMeshSelection& CurrentSelection = SelectionMechanic.GetCurrentSelection();
+		if (CurrentSelection.Mesh && CurrentSelection.Type == FDynamicMeshSelection::EType::Edge)
+		{
+			FDynamicMeshSelection UpdatedSelection = CurrentSelection;
+			UpdateSelectionEidsAfterMeshChange(UpdatedSelection, VidPairsIn);
+			SelectionMechanic.SetSelection(UpdatedSelection, false, false);
+		}
+	}
+
+	void GetVidPairsFromSelection(const FDynamicMeshSelection& SelectionIn, TArray<FIndex2i>& VidPairsOut)
+	{
+		VidPairsOut.Reset();
+		if (!SelectionIn.Mesh || SelectionIn.Type != FDynamicMeshSelection::EType::Edge)
+		{
+			// No vid pairs to add
+			return;
+		}
+
+		// Otherwise create the vid pairs out of eids
+		for (int32 Eid : SelectionIn.SelectedIDs)
+		{
+			VidPairsOut.Add(SelectionIn.Mesh->GetEdgeV(Eid));
+		}
+	}
+
 	/**
 	 * An undo/redo object for selection changes that, instead of operating directly on a selection
 	 * mechanic, instead operates on a context object that tools can use to route the request
@@ -37,6 +97,9 @@ namespace UVSelectToolLocals
 	 * to be undoable in different invocations of the tool, and the selection mechanic pointer
 	 * will not stay the same. However, the context object will stay the same, and we can register
 	 * to its delegate on each invocation.
+	 * 
+	 * The other thing that is different about this selection change object is that in cases of edge
+	 * selections, it uses stored vid pairs rather then eids, to deal with mesh changes that alter eids.
 	 */
 	class FSelectionChange : public FToolCommandChange
 	{
@@ -48,16 +111,27 @@ namespace UVSelectToolLocals
 		 *   the gizmo gets reset on the way forward to the current selection, which means we have to
 		 *   reset it to the old orientation on the way back (otherwise a rotated gizmo would end up
 		 *   losing its rotation on undo).
+		 * @param EdgeVidPairsBeforeIn
+		 * @param EdgeVidPairsAfterIn
 		 */
 		FSelectionChange(const FDynamicMeshSelection& SelectionBeforeIn,
 			const FDynamicMeshSelection& SelectionAfterIn,
 			bool bBroadcastOnSelectionChangedIn,
-			const FTransform& GizmoBeforeIn)
+			const FTransform& GizmoBeforeIn,
+			TUniquePtr<TArray<FIndex2i>> EdgeVidPairsBeforeIn,
+			TUniquePtr<TArray<FIndex2i>> EdgeVidPairsAfterIn
+			)
 			: SelectionBefore(SelectionBeforeIn)
 			, SelectionAfter(SelectionAfterIn)
 			, bBroadcastOnSelectionChanged(bBroadcastOnSelectionChangedIn)
 			, GizmoBefore(GizmoBeforeIn)
+			, EdgeVidPairsBefore(MoveTemp(EdgeVidPairsBeforeIn))
+			, EdgeVidPairsAfter(MoveTemp(EdgeVidPairsAfterIn))
 		{
+			// Make sure that for both selections, if we have a non-empty edge selection, we have vid pairs.
+			ensure(!(
+				(SelectionBefore.Mesh && SelectionBefore.Type == FDynamicMeshSelection::EType::Edge && !EdgeVidPairsBefore)
+				|| (SelectionAfter.Mesh && SelectionAfter.Type == FDynamicMeshSelection::EType::Edge && !EdgeVidPairsAfter)));
 		}
 
 		virtual void Apply(UObject* Object) override
@@ -65,6 +139,7 @@ namespace UVSelectToolLocals
 			UUVSelectToolChangeRouter* ChangeRouter = Cast<UUVSelectToolChangeRouter>(Object);
 			if (ensure(ChangeRouter) && ChangeRouter->CurrentSelectTool.IsValid())
 			{
+				UpdateSelectionEidsAfterMeshChange(SelectionAfter, EdgeVidPairsAfter.Get());
 				ChangeRouter->CurrentSelectTool->SetSelection(SelectionAfter, bBroadcastOnSelectionChanged);
 			}
 		}
@@ -74,6 +149,7 @@ namespace UVSelectToolLocals
 			UUVSelectToolChangeRouter* ChangeRouter = Cast<UUVSelectToolChangeRouter>(Object);
 			if (ensure(ChangeRouter) && ChangeRouter->CurrentSelectTool.IsValid())
 			{
+				UpdateSelectionEidsAfterMeshChange(SelectionBefore, EdgeVidPairsBefore.Get());
 				ChangeRouter->CurrentSelectTool->SetSelection(SelectionBefore, bBroadcastOnSelectionChanged);
 				if (bBroadcastOnSelectionChanged)
 				{
@@ -98,6 +174,9 @@ namespace UVSelectToolLocals
 		FDynamicMeshSelection SelectionAfter;
 		bool bBroadcastOnSelectionChanged;
 		FTransform GizmoBefore;
+
+		TUniquePtr<TArray<FIndex2i>> EdgeVidPairsBefore;
+		TUniquePtr<TArray<FIndex2i>> EdgeVidPairsAfter;
 	};
 
 	/**
@@ -201,12 +280,6 @@ UInteractiveTool* UUVSelectToolBuilder::BuildTool(const FToolBuilderState& Scene
 }
 
 // Tool property functions
-
-void  USelectToolActionPropertySet::Sew()
-{
-	PostAction(ESelectToolAction::Sew);
-}
-
 void  USelectToolActionPropertySet::IslandConformalUnwrap()
 {
 	PostAction(ESelectToolAction::IslandConformalUnwrap);
@@ -223,6 +296,8 @@ void USelectToolActionPropertySet::PostAction(ESelectToolAction Action)
 
 void UUVSelectTool::Setup()
 {
+	using namespace UVSelectToolLocals;
+
 	check(Targets.Num() > 0);
 
 	UInteractiveTool::Setup();
@@ -264,8 +339,21 @@ void UUVSelectTool::Setup()
 	SelectionMechanic->EmitSelectionChange = [this](const FDynamicMeshSelection& OldSelection,
 		const FDynamicMeshSelection& NewSelection, bool bBroadcastOnSelectionChangedIn)
 	{
+		TUniquePtr<TArray<FIndex2i>> VidPairsBefore;
+		TUniquePtr<TArray<FIndex2i>> VidPairsAfter;
+		if (OldSelection.Type == FDynamicMeshSelection::EType::Edge)
+		{
+			VidPairsBefore = MakeUnique<TArray<FIndex2i>>();
+			GetVidPairsFromSelection(OldSelection, *VidPairsBefore);
+		}
+		if (NewSelection.Type == FDynamicMeshSelection::EType::Edge)
+		{
+			VidPairsAfter = MakeUnique<TArray<FIndex2i>>();
+			GetVidPairsFromSelection(NewSelection, *VidPairsAfter);
+		}
 		EmitChangeAPI->EmitToolIndependentChange(ChangeRouter, MakeUnique<UVSelectToolLocals::FSelectionChange>(
-			OldSelection, NewSelection, bBroadcastOnSelectionChangedIn, TransformGizmo->GetGizmoTransform()), 
+			OldSelection, NewSelection, bBroadcastOnSelectionChangedIn, TransformGizmo->GetGizmoTransform(),
+			MoveTemp(VidPairsBefore), MoveTemp(VidPairsAfter)),
 			LOCTEXT("SelectionChangeMessage", "Selection Change"));
 	};
 
@@ -305,6 +393,7 @@ void UUVSelectTool::Setup()
 	{
 		Targets[i]->OnCanonicalModified.AddWeakLambda(this, [this, i]
 		(UUVEditorToolMeshInput* InputObject, const UUVEditorToolMeshInput::FCanonicalModifiedInfo& Info) {
+			UpdateSelectionEidsAfterMeshChange(*SelectionMechanic, &CurrentSelectionVidPairs);
 			UpdateGizmo();
 			SelectionMechanic->RebuildDrawnElements(TransformGizmo->GetGizmoTransform());
 		});
@@ -458,12 +547,16 @@ void UUVSelectTool::ConfigureSelectionModeFromControls()
 
 void UUVSelectTool::OnSelectionChanged()
 {
+	using namespace UVSelectToolLocals;
+
 	const FDynamicMeshSelection& Selection = SelectionMechanic->GetCurrentSelection();
 
+	GetVidPairsFromSelection(Selection, CurrentSelectionVidPairs);
+
 	SelectionTargetIndex = -1;
 	MovingVids.Reset();
 	SelectedTids.Reset();
-	BoundaryEids.Reset();
+	LivePreviewBoundaryEids.Reset();
 
 	if (!Selection.IsEmpty())
 	{
@@ -510,7 +603,7 @@ void UUVSelectTool::OnSelectionChanged()
 					{
 						if (EdgeTids[j] != Tid && !Selection.SelectedIDs.Contains(EdgeTids[j]))
 						{
-							BoundaryEids.Add(TriEids[i]);
+							LivePreviewBoundaryEids.Add(TriEids[i]);
 							break;
 						}
 					}
@@ -589,7 +682,7 @@ void UUVSelectTool::UpdateLivePreviewLines()
 		FTransform MeshTransform = Targets[SelectionTargetIndex]->AppliedPreview->PreviewMesh->GetTransform();
 		const FDynamicMesh3* LivePreviewMesh = Targets[SelectionTargetIndex]->AppliedCanonical.Get();
 
-		for (int32 Eid : BoundaryEids)
+		for (int32 Eid : LivePreviewBoundaryEids)
 		{
 			FVector3d Vert1, Vert2;
 			LivePreviewMesh->GetEdgeV(Eid, Vert1, Vert2);
@@ -775,9 +868,76 @@ void UUVSelectTool::ApplyAction(ESelectToolAction ActionType)
 			EmitChangeAPI->EndUndoTransaction();
 		}
 		break;
+	case ESelectToolAction::Split:
+		ApplySplit();
 	default:
 		break;
 	}
 }
 
+void UUVSelectTool::ApplySplit()
+{
+	using namespace UVSelectToolLocals;
+
+	const FDynamicMeshSelection& Selection = SelectionMechanic->GetCurrentSelection();
+
+	if (Selection.IsEmpty() || Selection.Type != FDynamicMeshSelection::EType::Edge)
+	{
+		return;
+	}
+
+	if (!ensure(SelectionTargetIndex >= 0))
+	{
+		return;
+	}
+	UUVEditorToolMeshInput* Target = Targets[SelectionTargetIndex];
+	
+	// Gather up the corresponding edge IDs in the applied (3d) mesh.
+	TSet<int32> AppliedEidSet;
+	for (int32 Eid : Selection.SelectedIDs)
+	{
+		FIndex2i EdgeUnwrapVids = Selection.Mesh->GetEdgeV(Eid);
+
+		int32 AppliedEid = Target->AppliedCanonical->FindEdge(
+			Target->UnwrapVidToAppliedVid(EdgeUnwrapVids.A),
+			Target->UnwrapVidToAppliedVid(EdgeUnwrapVids.B));
+
+		if (ensure(AppliedEid != IndexConstants::InvalidID))
+		{
+			AppliedEidSet.Add(AppliedEid);
+		}
+	}
+
+	// Perform the cut in the overlay
+	FUVEditResult UVEditResult;
+	FDynamicMeshUVEditor UVEditor(Target->AppliedCanonical.Get(),
+		Target->UVLayerIndex, false);
+	UVEditor.CreateSeamsAtEdges(AppliedEidSet, &UVEditResult);
+
+	// Figure out the triangles that need to be saved in the unwrap for undo
+	TSet<int32> TidSet;
+	for (int32 UnwrapVid : UVEditResult.NewUVElements)
+	{
+		TArray<int32> VertTids;
+		Target->AppliedCanonical->GetVtxTriangles(Target->UnwrapVidToAppliedVid(UnwrapVid), VertTids);
+		TidSet.Append(VertTids);
+	}
+
+	FDynamicMeshChangeTracker ChangeTracker(Target->UnwrapCanonical.Get());
+	ChangeTracker.BeginChange();
+	ChangeTracker.SaveTriangles(TidSet, true);
+
+	// Perform the update
+	TArray<int32> AppliedTids = TidSet.Array();
+	Target->UpdateAllFromAppliedCanonical(&UVEditResult.NewUVElements, &AppliedTids, &AppliedTids);
+	AABBTrees[SelectionTargetIndex]->Build();
+
+	// Emit transaction
+	const FText TransactionName(LOCTEXT("ApplySplitTransactionName", "Split Edges"));
+	EmitChangeAPI->EmitToolIndependentUnwrapCanonicalChange(
+		Target, ChangeTracker.EndChange(), TransactionName);
+
+	UpdateSelectionEidsAfterMeshChange(*SelectionMechanic, &CurrentSelectionVidPairs);
+}
+
 #undef LOCTEXT_NAMESPACE
diff --git a/Engine/Plugins/Experimental/UVEditor/Source/UVEditorTools/Public/UVSelectTool.h b/Engine/Plugins/Experimental/UVEditor/Source/UVEditorTools/Public/UVSelectTool.h
--- a/Engine/Plugins/Experimental/UVEditor/Source/UVEditorTools/Public/UVSelectTool.h
+++ b/Engine/Plugins/Experimental/UVEditor/Source/UVEditorTools/Public/UVSelectTool.h
@@ -83,7 +83,8 @@ enum class ESelectToolAction
 	NoAction,
 
 	Sew,
-	IslandConformalUnwrap
+	Split,
+	IslandConformalUnwrap,
 };
 
 UCLASS()
@@ -97,14 +98,17 @@ public:
 	void Initialize(UUVSelectTool* ParentToolIn) { ParentTool = ParentToolIn; }
 	void PostAction(ESelectToolAction Action);
 
-	/* Sew selected boundary edges together topologically in the UV set */
+	/** Sew edges. The red edges will be sewn to the green edges */
 	UFUNCTION(CallInEditor, Category = Actions)
-	void Sew();
+	void Sew() { PostAction(ESelectToolAction::Sew); }
+
+	/** Split given edges/verts */
+	UFUNCTION(CallInEditor, Category = Actions)
+	void Split() { PostAction(ESelectToolAction::Split); };
 
 	/* Apply a conformal unwrap to the selected UV islands */
 	UFUNCTION(CallInEditor, Category = Actions)
 	void IslandConformalUnwrap();
-
 };
 
 
@@ -216,11 +220,17 @@ protected:
 	TArray<int32> SelectedTids;
 	TArray<FVector3d> MovingVertOriginalPositions;
 	int32 SelectionTargetIndex;
-	TArray<int32> BoundaryEids;
+	TArray<int32> LivePreviewBoundaryEids;
 	
+	// When selecting edges, used to hold the edges as pairs of vids, because the eids change
+	// during undo/redo and other topological operations.
+	TArray<UE::Geometry::FIndex2i> CurrentSelectionVidPairs;
+
 	ESelectToolAction PendingAction;
 
 	void ApplyAction(ESelectToolAction ActionType);
+	void ApplySplit();
+
 	const float LivePreviewHighlightThickness = 2.0;
 	const float LivePreviewHighlightDepthOffset = 0.5;
 };
diff --git a/Engine/Plugins/Runtime/GeometryProcessing/Source/DynamicMesh/Private/Parameterization/DynamicMeshUVEditor.cpp b/Engine/Plugins/Runtime/GeometryProcessing/Source/DynamicMesh/Private/Parameterization/DynamicMeshUVEditor.cpp
--- a/Engine/Plugins/Runtime/GeometryProcessing/Source/DynamicMesh/Private/Parameterization/DynamicMeshUVEditor.cpp
+++ b/Engine/Plugins/Runtime/GeometryProcessing/Source/DynamicMesh/Private/Parameterization/DynamicMeshUVEditor.cpp
@@ -779,85 +779,96 @@ static int32 FindUVElementForVertex(const FDynamicMesh3* Mesh, const FDynamicMes
 
 
 
-bool FDynamicMeshUVEditor::CreateSeamAlongVertexPath(const TArray<int32>& VertexPath, FUVEditResult* Result)
+bool FDynamicMeshUVEditor::CreateSeamsAtEdges(const TSet<int32>& EidsToMakeIntoSeams, FUVEditResult* Result)
 {
-	// todo: handle paths that are loops (may just work?)
-	check(VertexPath[0] != VertexPath.Last());
-
-	// construct list of sequential mesh edges, and set of seam edges, for VertexPath
-	TArray<int32> EdgePath;
-	TSet<int32> SeamEdges;
-	int32 NumVerts = VertexPath.Num();
-	for (int32 k = 0; k < NumVerts - 1; ++k)
+	/**
+	 * @param TidsUntilNextSeam Tids in a fan away from Eid around Vid until the next seam. Will
+	 *  be the entire one-ring if the given Eid is the only seam attached to Vid.
+	 */
+	auto DoesVertHaveAnotherSeamAttached = [this, &EidsToMakeIntoSeams](int32 Eid, int32 Vid, TArray<int32>& TidsUntilNextSeam)
 	{
-		int32 FoundEdgeID = Mesh->FindEdge(VertexPath[k], VertexPath[k+1]);
-		if (Mesh->IsEdge(FoundEdgeID) == false)
-		{
-			return false;
-		}
+		TidsUntilNextSeam.Reset();
+
+		FIndex2i EdgeTids = Mesh->GetEdgeT(Eid);
+		int32 CurrentTid = EdgeTids.A;
+		int32 PreviousTid = EdgeTids.B;
 
-		EdgePath.Add(FoundEdgeID);
-		if (UVOverlay->IsSeamEdge(FoundEdgeID))
+		int32 MaxNumTids = Mesh->GetVtxEdgeCount(Vid);
+		while (true) 
 		{
-			SeamEdges.Add(FoundEdgeID);
-		}
-	}
+			TidsUntilNextSeam.Add(CurrentTid);
+			if (TidsUntilNextSeam.Num() > MaxNumTids) // sanity check
+			{
+				ensure(false);
+				return false;
+			}
 
-	bool bStartIsSeamVtx = UVOverlay->IsSeamVertex(VertexPath[0]);
-	bool bEndIsSeamVtx = UVOverlay->IsSeamVertex(VertexPath.Last());
+			FIndex3i NextTriResult = UE::Geometry::FindNextAdjacentTriangleAroundVtx(Mesh, Vid, CurrentTid, PreviousTid,
+				[this, &EidsToMakeIntoSeams, Eid](int32 Tri0, int32 Tri1, int32 EidBetween) {
+					// We stop at seam or at future seam not including our starting edge (if we walk across the starting
+					// edge we will stop at a later check- we just don't want to think that we reached a true seam).
+					return UVOverlay->AreTrianglesConnected(Tri0, Tri1) && !(EidBetween != Eid && EidsToMakeIntoSeams.Contains(EidBetween)); }
+			);
+			int32 NextTid = NextTriResult.A;
 
-	// handle start vert. If it's on a boundary/seam we split it, otherwise we do not (next-vtx split will create seam)
-	if (bStartIsSeamVtx && SeamEdges.Contains(EdgePath[0]) == false)
-	{
-		TArray<int32> SplitSets[2];
-		bool bSetsOK = FindSeamTriSplitSets_BoundaryVtx(Mesh, UVOverlay, VertexPath[0], EdgePath[0], SplitSets);
-		check(bSetsOK);
-		int32 ElementID = FindUVElementForVertex(Mesh, UVOverlay, VertexPath[0], EdgePath[0]);
-		check(ElementID != IndexConstants::InvalidID);
-		int32 NewElementID = UVOverlay->SplitElement(ElementID, SplitSets[0]);
-		if (Result)
+			// See if we've come to a (current or future) seam or wrapped around
+			if (NextTid == IndexConstants::InvalidID)
+			{
+				return true;
+			}
+			else if (NextTid == EdgeTids.A)
+			{
+				return false;
+			}
+			PreviousTid = CurrentTid;
+			CurrentTid = NextTid;
+		}
+	};
+
+	auto SplitEdgeVertElement = [this, Result](int32 Eid, int32 Vid, const TArray<int32>& TidsToModify) {
+		int32 ElementID = FindUVElementForVertex(Mesh, UVOverlay, Vid, Eid);
+		if (ensure(ElementID != IndexConstants::InvalidID))
 		{
-			Result->NewUVElements.Add(NewElementID);
+			int32 NewElementID = UVOverlay->SplitElement(ElementID, TidsToModify);
+			if (Result)
+			{
+				Result->NewUVElements.Add(NewElementID);
+			}
 		}
-	}
+	};
 
-	// handle end-vert similarly
-	if (bEndIsSeamVtx && SeamEdges.Contains(EdgePath.Last()) == false)
+	for (int32 Eid : EidsToMakeIntoSeams)
 	{
-		TArray<int32> SplitSets[2];
-		bool bSetsOK = FindSeamTriSplitSets_BoundaryVtx(Mesh, UVOverlay, VertexPath.Last(), EdgePath.Last(), SplitSets);
-		check(bSetsOK);
-		int32 ElementID = FindUVElementForVertex(Mesh, UVOverlay, VertexPath.Last(), EdgePath.Last());
-		check(ElementID != IndexConstants::InvalidID);
-		int32 NewElementID = UVOverlay->SplitElement(ElementID, SplitSets[0]);
-		if (Result)
+		// If the edge is already a seam, don't need to do anything for it
+		if (UVOverlay->IsSeamEdge(Eid))
 		{
-			Result->NewUVElements.Add(NewElementID);
+			continue;
 		}
-	}
 
-	// walk through intermediate verts and cut connected island at each one
-	for (int32 k = 1; k < NumVerts - 1; ++k)
-	{
-		// todo: need to check here that we have not already separated the two SplitSets?
+		FIndex2i EdgeVids = Mesh->GetEdgeV(Eid);
+		TArray<int32> TidsToNextSeamA, TidsToNextSeamB;
+		bool bSeamCreated = false;
 
-		int32 PrevEdge = EdgePath[k-1];
-		int32 NextEdge = EdgePath[k];
+		// Perform the two tests before we split either vert. A vert has to get split if it has a
+		// seam attached, because splitting the other vert would create a bowtie.
+		bool VertANeedsSplitting = DoesVertHaveAnotherSeamAttached(Eid, EdgeVids.A, TidsToNextSeamA);
+		bool VertBNeedsSplitting = DoesVertHaveAnotherSeamAttached(Eid, EdgeVids.B, TidsToNextSeamB);
 
-		if (Mesh->IsBoundaryEdge(PrevEdge) || Mesh->IsBoundaryEdge(NextEdge))
+		// If neither absolutely has to get split, then one needs to get split anyway so that we
+		// make the edge into a seam.
+		if (!VertANeedsSplitting && !VertBNeedsSplitting)
 		{
-			continue;
+			VertANeedsSplitting = true;
+			TidsToNextSeamA.SetNum((TidsToNextSeamA.Num() + 1) / 2); // Go halfway around, rounding up
 		}
 
-		TArray<int32> SplitSets[2];
-		bool bSetsOK = FindSeamTriSplitSets_InteriorVtx(Mesh, UVOverlay, VertexPath[k], PrevEdge, NextEdge, SplitSets);
-		check(bSetsOK);
-		int32 ElementID = FindUVElementForVertex(Mesh, UVOverlay, VertexPath[k], PrevEdge);
-		check(ElementID != IndexConstants::InvalidID);
-		int32 NewElementID = UVOverlay->SplitElement(ElementID, SplitSets[0]);
-		if (Result)
+		if (VertANeedsSplitting)
 		{
-			Result->NewUVElements.Add(NewElementID);
+			SplitEdgeVertElement(Eid, EdgeVids.A, TidsToNextSeamA);
+		}
+		if (VertBNeedsSplitting)
+		{
+			SplitEdgeVertElement(Eid, EdgeVids.B, TidsToNextSeamB);
 		}
 	}
 
@@ -866,7 +877,6 @@ bool FDynamicMeshUVEditor::CreateSeamAlongVertexPath(const TArray<int32>& Vertex
 
 
 
-
 void FDynamicMeshUVEditor::SetTriangleUVsFromBoxProjection(
 	const TArray<int32>& Triangles, 
 	TFunctionRef<FVector3d(const FVector3d&)> PointTransform, 
diff --git a/Engine/Plugins/Runtime/GeometryProcessing/Source/DynamicMesh/Public/Parameterization/DynamicMeshUVEditor.h b/Engine/Plugins/Runtime/GeometryProcessing/Source/DynamicMesh/Public/Parameterization/DynamicMeshUVEditor.h
--- a/Engine/Plugins/Runtime/GeometryProcessing/Source/DynamicMesh/Public/Parameterization/DynamicMeshUVEditor.h
+++ b/Engine/Plugins/Runtime/GeometryProcessing/Source/DynamicMesh/Public/Parameterization/DynamicMeshUVEditor.h
@@ -171,15 +171,18 @@ public:
 	 */
 	bool SetTriangleUVsFromFreeBoundaryConformal(const TArray<int32>& Triangles, bool bUseExistingUVTopology, FUVEditResult* Result = nullptr);
 
-
 	/**
-	 * Cut existing UV topolgy with a path of vertices. This allows for creating partial seams/darts, interior cuts, etc.
-	 * @param VertexPath sequential list of edge-connected vertices
+	 * Cut existing UV topolgy with a set of edges. This allows for creating partial seams/darts, interior cuts, etc.
+	 * 
+	 * Avoids creating bowties. In cases where an edge is not next to any present or future seams/borders, some
+	 * adjacent edge will be picked to be made into a seam as well, since it's impossible to make the original
+	 * into a seam otherwise.
+
+	 * @param EidsToMakeIntoSeams list of edges to turn into seams
 	 * @param Result if non-null, list of new UV elements created along the path will be stored here (not ordered)
 	 * @return true on success
 	 */
-	bool CreateSeamAlongVertexPath(const TArray<int32>& VertexPath, FUVEditResult* Result = nullptr);
-
+	bool CreateSeamsAtEdges(const TSet<int32>& EidsToMakeIntoSeams, FUVEditResult* Result = nullptr);
 
 	/**
 	 * Set UVs by box projection. Triangles will be grouped to "best" box face
