UVEditor: Added Seam tool. Also added support for nested accept/cancel in UV editor, and for "complete" button.

#rnx

diff --git a/Engine/Plugins/Experimental/UVEditor/Source/UVEditor/Private/UVEditorCommands.cpp b/Engine/Plugins/Experimental/UVEditor/Source/UVEditor/Private/UVEditorCommands.cpp
--- a/Engine/Plugins/Experimental/UVEditor/Source/UVEditor/Private/UVEditorCommands.cpp
+++ b/Engine/Plugins/Experimental/UVEditor/Source/UVEditor/Private/UVEditorCommands.cpp
@@ -27,10 +27,11 @@ void FUVEditorCommands::RegisterCommands()
 	UI_COMMAND(BeginLayoutTool, "Layout", "Switch to layout tool", EUserInterfaceActionType::ToggleButton, FInputChord());
 	UI_COMMAND(BeginParameterizeMeshTool, "Auto UV", "Switch to Auto UV tool", EUserInterfaceActionType::ToggleButton, FInputChord());
 	UI_COMMAND(BeginChannelEditTool, "Channel Edit", "Switch to Channel Edit tool", EUserInterfaceActionType::ToggleButton, FInputChord());
+	UI_COMMAND(BeginSeamTool, "Seam", "Switch to Seam tool", EUserInterfaceActionType::ToggleButton, FInputChord());
 
 	// These allow us to link up to pressed keys
-	UI_COMMAND(AcceptActiveTool, "Accept", "Accept the active tool", EUserInterfaceActionType::Button, FInputChord(EKeys::Enter));
-	UI_COMMAND(CancelActiveTool, "Cancel", "Cancel the active tool", EUserInterfaceActionType::Button, FInputChord(EKeys::Escape));
+	UI_COMMAND(AcceptOrCompleteActiveTool, "Accept", "Accept the active tool", EUserInterfaceActionType::Button, FInputChord(EKeys::Enter));
+	UI_COMMAND(CancelOrCompleteActiveTool, "Cancel", "Cancel the active tool", EUserInterfaceActionType::Button, FInputChord(EKeys::Escape));
 
 	UI_COMMAND(EnableOrbitCamera, "Orbit", "Enable Orbit Camera", EUserInterfaceActionType::ToggleButton, FInputChord());
 	UI_COMMAND(EnableFlyCamera, "Fly", "Enable Fly Camera", EUserInterfaceActionType::ToggleButton, FInputChord());
diff --git a/Engine/Plugins/Experimental/UVEditor/Source/UVEditor/Private/UVEditorMode.cpp b/Engine/Plugins/Experimental/UVEditor/Source/UVEditor/Private/UVEditorMode.cpp
--- a/Engine/Plugins/Experimental/UVEditor/Source/UVEditor/Private/UVEditorMode.cpp
+++ b/Engine/Plugins/Experimental/UVEditor/Source/UVEditor/Private/UVEditorMode.cpp
@@ -25,6 +25,7 @@
 #include "UVEditorLayoutTool.h"
 #include "UVEditorParameterizeMeshTool.h"
 #include "UVEditorLayerEditTool.h"
+#include "UVEditorSeamTool.h"
 #include "UVSelectTool.h"
 #include "UVEditorModeToolkit.h"
 #include "UVEditorSubsystem.h"
@@ -179,6 +180,10 @@ void UUVEditorMode::RegisterTools()
 	UUVEditorChannelEditToolBuilder* UVEditorChannelEditToolBuilder = NewObject<UUVEditorChannelEditToolBuilder>();
 	UVEditorChannelEditToolBuilder->Targets = &ToolInputObjects;
 	RegisterTool(CommandInfos.BeginChannelEditTool, TEXT("UVChannelEditTool"), UVEditorChannelEditToolBuilder);
+
+	UUVEditorSeamToolBuilder* UVEditorSeamToolBuilder = NewObject<UUVEditorSeamToolBuilder>();
+	UVEditorSeamToolBuilder->Targets = &ToolInputObjects;
+	RegisterTool(CommandInfos.BeginSeamTool, TEXT("UVSeamTool"), UVEditorSeamToolBuilder);
 }
 
 void UUVEditorMode::AddDisplayedPropertySet(const TObjectPtr<UInteractiveToolPropertySet>& PropertySet)
@@ -209,24 +214,69 @@ void UUVEditorMode::BindCommands()
 
 	// Hook up to Enter/Esc key presses
 	CommandList->MapAction(
-		CommandInfos.AcceptActiveTool,
-		FExecuteAction::CreateLambda([this]() { 
+		CommandInfos.AcceptOrCompleteActiveTool,
+		FExecuteAction::CreateLambda([this]() {
+			// Give the tool a chance to take the nested accept first
+			if (GetToolManager()->HasAnyActiveTool())
+			{
+				UInteractiveTool* Tool = GetToolManager()->GetActiveTool(EToolSide::Mouse);
+				IInteractiveToolNestedAcceptCancelAPI* AcceptAPI = Cast<IInteractiveToolNestedAcceptCancelAPI>(Tool);
+				if (AcceptAPI && AcceptAPI->SupportsNestedAcceptCommand() && AcceptAPI->CanCurrentlyNestedAccept())
+				{
+					if (AcceptAPI->ExecuteNestedAcceptCommand())
+					{
+						return;
+					}
+				}
+			}
 			GetInteractiveToolsContext()->EndTool(EToolShutdownType::Accept); 
 			ActivateDefaultTool();
 			}),
-		FCanExecuteAction::CreateLambda([this]() { return GetInteractiveToolsContext()->CanAcceptActiveTool(); }),
+		FCanExecuteAction::CreateLambda([this]() {
+			if (GetInteractiveToolsContext()->CanAcceptActiveTool() || GetInteractiveToolsContext()->CanCompleteActiveTool())
+			{
+				return true;
+			}
+			// If we can't currently accept, may still be able to pass down to tool
+			UInteractiveTool* Tool = GetToolManager()->GetActiveTool(EToolSide::Mouse);
+			IInteractiveToolNestedAcceptCancelAPI* AcceptAPI = Cast<IInteractiveToolNestedAcceptCancelAPI>(Tool);
+			return AcceptAPI && AcceptAPI->SupportsNestedAcceptCommand() && AcceptAPI->CanCurrentlyNestedAccept(); 
+			}),
 		FGetActionCheckState(),
 		FIsActionButtonVisible::CreateLambda([this]() {return GetInteractiveToolsContext()->ActiveToolHasAccept(); }),
 		EUIActionRepeatMode::RepeatDisabled
 	);
 
 	CommandList->MapAction(
-		CommandInfos.CancelActiveTool,
+		CommandInfos.CancelOrCompleteActiveTool,
 		FExecuteAction::CreateLambda([this]() { 
+			// Give the tool a chance to take the nested cancel first
+			if (GetToolManager()->HasAnyActiveTool())
+			{
+				UInteractiveTool* Tool = GetToolManager()->GetActiveTool(EToolSide::Mouse);
+				IInteractiveToolNestedAcceptCancelAPI* CancelAPI = Cast<IInteractiveToolNestedAcceptCancelAPI>(Tool);
+				if (CancelAPI && CancelAPI->SupportsNestedCancelCommand() && CancelAPI->CanCurrentlyNestedCancel())
+				{
+					if (CancelAPI->ExecuteNestedCancelCommand())
+					{
+						return;
+					}
+				}
+			}
+
 			GetInteractiveToolsContext()->EndTool(EToolShutdownType::Cancel);
 			ActivateDefaultTool();
 			}),
-		FCanExecuteAction::CreateLambda([this]() { return GetInteractiveToolsContext()->CanCancelActiveTool(); }),
+		FCanExecuteAction::CreateLambda([this]() { 
+			if (GetInteractiveToolsContext()->CanCancelActiveTool() || GetInteractiveToolsContext()->CanCompleteActiveTool())
+			{
+				return true;
+			}
+			// If we can't currently cancel, may still be able to pass down to tool
+			UInteractiveTool* Tool = GetToolManager()->GetActiveTool(EToolSide::Mouse);
+			IInteractiveToolNestedAcceptCancelAPI* CancelAPI = Cast<IInteractiveToolNestedAcceptCancelAPI>(Tool);
+			return CancelAPI && CancelAPI->SupportsNestedCancelCommand() && CancelAPI->CanCurrentlyNestedCancel();
+			}),
 		FGetActionCheckState(),
 		FIsActionButtonVisible::CreateLambda([this]() {return GetInteractiveToolsContext()->ActiveToolHasAccept(); }),
 		EUIActionRepeatMode::RepeatDisabled
diff --git a/Engine/Plugins/Experimental/UVEditor/Source/UVEditor/Private/UVEditorModeToolkit.cpp b/Engine/Plugins/Experimental/UVEditor/Source/UVEditor/Private/UVEditorModeToolkit.cpp
--- a/Engine/Plugins/Experimental/UVEditor/Source/UVEditor/Private/UVEditorModeToolkit.cpp
+++ b/Engine/Plugins/Experimental/UVEditor/Source/UVEditor/Private/UVEditorModeToolkit.cpp
@@ -189,6 +189,7 @@ void FUVEditorModeToolkit::Init(const TSharedPtr<IToolkitHost>& InitToolkitHost,
 	ToolbarBuilder.AddToolBarButton(CommandInfos.BeginLayoutTool);
 	ToolbarBuilder.AddToolBarButton(CommandInfos.BeginParameterizeMeshTool);
 	ToolbarBuilder.AddToolBarButton(CommandInfos.BeginChannelEditTool);
+	ToolbarBuilder.AddToolBarButton(CommandInfos.BeginSeamTool);
 
 	// Hook in the tool palette
 	ToolButtonsContainer->SetContent(ToolbarBuilder.MakeWidget());
@@ -217,8 +218,7 @@ void FUVEditorModeToolkit::Init(const TSharedPtr<IToolkitHost>& InitToolkitHost,
 		BackgroundDetailsContainer->SetContent(BackgroundDetailsView.ToSharedRef());
 	}
 	
-	// Set up the overlay. Larely copied from ModelingToolsEditorModeToolkit, except
-	// that we don't have a "complete" option since we don't have a non-tool state.
+	// Set up the overlay. Largely copied from ModelingToolsEditorModeToolkit.
 	// TODO: We could put some of the shared code in some common place.
 	SAssignNew(ViewportOverlayWidget, SHorizontalBox)
 
@@ -287,21 +287,27 @@ void FUVEditorModeToolkit::Init(const TSharedPtr<IToolkitHost>& InitToolkitHost,
 				.Visibility_Lambda([this]() { return GetScriptableEditorMode()->GetInteractiveToolsContext()->ActiveToolHasAccept() ? EVisibility::Visible : EVisibility::Collapsed; })
 			]
 
-			// Currently haven't needed a complete button, but this is here in case that changes.
-			//+SHorizontalBox::Slot()
-			//.AutoWidth()
-			//.Padding(FMargin(2.0, 0.f, 0.f, 0.f))
-			//[
-			//	SNew(SButton)
-			//	.ButtonStyle(FAppStyle::Get(), "PrimaryButton")
-			//	.TextStyle( FAppStyle::Get(), "DialogButtonText" )
-			//	.Text(LOCTEXT("OverlayComplete", "Complete"))
-			//	.ToolTipText(LOCTEXT("OverlayCompleteTooltip", "Exit the active Tool [Enter]"))
-			//	.HAlign(HAlign_Center)
-			//	.OnClicked_Lambda([this]() { GetScriptableEditorMode()->GetInteractiveToolsContext()->EndTool(EToolShutdownType::Completed); return FReply::Handled(); })
-			//	.IsEnabled_Lambda([this]() { return GetScriptableEditorMode()->GetInteractiveToolsContext()->CanCompleteActiveTool(); })
-			//	.Visibility_Lambda([this]() { return GetScriptableEditorMode()->GetInteractiveToolsContext()->CanCompleteActiveTool() ? EVisibility::Visible : EVisibility::Collapsed; })
-			//]
+			+ SHorizontalBox::Slot()
+				.AutoWidth()
+				.Padding(FMargin(2.0, 0.f, 0.f, 0.f))
+				[
+					SNew(SButton)
+					.ButtonStyle(FAppStyle::Get(), "PrimaryButton")
+				.TextStyle(FAppStyle::Get(), "DialogButtonText")
+				.Text(LOCTEXT("OverlayComplete", "Complete"))
+				.ToolTipText(LOCTEXT("OverlayCompleteTooltip", "Exit the active Tool [Enter]"))
+				.HAlign(HAlign_Center)
+				.OnClicked_Lambda([this]() { 
+					GetScriptableEditorMode()->GetInteractiveToolsContext()->EndTool(EToolShutdownType::Completed); 
+					Cast<UUVEditorMode>(GetScriptableEditorMode())->ActivateDefaultTool();
+					return FReply::Handled(); 
+					})
+				.IsEnabled_Lambda([this]() {
+					UUVEditorMode* Mode = Cast<UUVEditorMode>(GetScriptableEditorMode());
+					return GetScriptableEditorMode()->GetInteractiveToolsContext()->CanCompleteActiveTool();
+				})
+				.Visibility_Lambda([this]() { return GetScriptableEditorMode()->GetInteractiveToolsContext()->CanCompleteActiveTool() ? EVisibility::Visible : EVisibility::Collapsed; })
+			]
 		]	
 	];
 
@@ -347,19 +353,17 @@ void FUVEditorModeToolkit::OnToolStarted(UInteractiveToolManager* Manager, UInte
 
 	ActiveToolName = Tool->GetToolInfo().ToolDisplayName;
 
-	if (GetScriptableEditorMode()->GetInteractiveToolsContext()->ActiveToolHasAccept())
-	{
-		// Add the accept/cancel overlay
-		GetToolkitHost()->AddViewportOverlayWidget(ViewportOverlayWidget.ToSharedRef());
-	}
-	
-	// Issue a tool start transaction unless we are starting the default tool, because we can't
-	// undo or revert out of the default tool.
 	UUVEditorMode* Mode = Cast<UUVEditorMode>(GetScriptableEditorMode());
 	if (!Mode->IsDefaultToolActive())
 	{
+		// Issue a tool start transaction unless we are starting the default tool, because we can't
+		// undo or revert out of the default tool.
 		Mode->GetInteractiveToolsContext()->GetTransactionAPI()->AppendChange(
 			Mode, MakeUnique<UVEditorModeToolkitLocals::FUVEditorBeginToolChange>(), LOCTEXT("ActivateTool", "Activate Tool"));
+
+		// Add the accept/cancel overlay. Again, unless we're in the default tool, which we don't leave
+		// except by activating another tool.
+		GetToolkitHost()->AddViewportOverlayWidget(ViewportOverlayWidget.ToSharedRef());
 	}
 }
 
diff --git a/Engine/Plugins/Experimental/UVEditor/Source/UVEditor/Public/UVEditorCommands.h b/Engine/Plugins/Experimental/UVEditor/Source/UVEditor/Public/UVEditorCommands.h
--- a/Engine/Plugins/Experimental/UVEditor/Source/UVEditor/Public/UVEditorCommands.h
+++ b/Engine/Plugins/Experimental/UVEditor/Source/UVEditor/Public/UVEditorCommands.h
@@ -22,9 +22,10 @@ public:
 	TSharedPtr<FUICommandInfo> BeginLayoutTool;
 	TSharedPtr<FUICommandInfo> BeginParameterizeMeshTool;
 	TSharedPtr<FUICommandInfo> BeginChannelEditTool;
+	TSharedPtr<FUICommandInfo> BeginSeamTool;
 
-	TSharedPtr<FUICommandInfo> AcceptActiveTool;
-	TSharedPtr<FUICommandInfo> CancelActiveTool;
+	TSharedPtr<FUICommandInfo> AcceptOrCompleteActiveTool;
+	TSharedPtr<FUICommandInfo> CancelOrCompleteActiveTool;
 
 	TSharedPtr<FUICommandInfo> EnableOrbitCamera;
 	TSharedPtr<FUICommandInfo> EnableFlyCamera;
diff --git a/Engine/Plugins/Experimental/UVEditor/Source/UVEditorTools/Private/UVEditorLayerEditTool.cpp b/Engine/Plugins/Experimental/UVEditor/Source/UVEditorTools/Private/UVEditorLayerEditTool.cpp
--- a/Engine/Plugins/Experimental/UVEditor/Source/UVEditorTools/Private/UVEditorLayerEditTool.cpp
+++ b/Engine/Plugins/Experimental/UVEditor/Source/UVEditorTools/Private/UVEditorLayerEditTool.cpp
@@ -541,11 +541,6 @@ void UUVEditorChannelEditTool::ApplyVisbleChannelChange()
 	}
 }
 
-bool UUVEditorChannelEditTool::CanAccept() const
-{
-	return true;
-}
-
 void UUVEditorChannelEditTool::OnPropertyModified(UObject* PropertySet, FProperty* Property)
 {
 
diff --git a/Engine/Plugins/Experimental/UVEditor/Source/UVEditorTools/Private/UVEditorSeamTool.cpp b/Engine/Plugins/Experimental/UVEditor/Source/UVEditorTools/Private/UVEditorSeamTool.cpp
new file mode 100644
--- /dev/null
+++ b/Engine/Plugins/Experimental/UVEditor/Source/UVEditorTools/Private/UVEditorSeamTool.cpp
@@ -0,0 +1,1126 @@
+// Copyright Epic Games, Inc. All Rights Reserved.
+
+#include "UVEditorSeamTool.h"
+
+#include "Algo/Reverse.h"
+#include "BaseBehaviors/SingleClickBehavior.h"
+#include "BaseBehaviors/MouseHoverBehavior.h"
+#include "ContextObjectStore.h"
+#include "Drawing/PreviewGeometryActor.h"
+#include "DynamicMesh/DynamicMeshChangeTracker.h"
+#include "InputBehaviorSet.h"
+#include "InputRouter.h"
+#include "MeshOpPreviewHelpers.h"
+#include "Parameterization/DynamicMeshUVEditor.h"
+#include "Parameterization/MeshDijkstra.h"
+#include "ToolTargets/UVEditorToolMeshInput.h"
+#include "ToolSceneQueriesUtil.h"
+#include "UVToolContextObjects.h"
+
+#define LOCTEXT_NAMESPACE "UUVEditorSeamTool"
+
+using namespace UE::Geometry;
+
+namespace UVEditorSeamToolLocals
+{
+	const FText EditSeamTransactionName(LOCTEXT("EditSeamTransactionName", "Edit Working Seam"));
+	const FText ApplySeamTransactionName(LOCTEXT("ApplySeamTransactionName", "Apply Seam"));
+
+	const FString& HoverPointSetID(TEXT("HoverPointSet"));
+	const FString& HoverLineSetID(TEXT("HoverLineSet"));
+
+	const FString& LockedPointSetID(TEXT("LockedPointSet"));
+	const FString& LockedLineSetID(TEXT("LockedLineSet"));
+
+	const FString& StartPointSetID(TEXT("StartPointSet"));
+	
+	const FString& ExistingSeamsID(TEXT("SeamLineSet"));
+
+	const FColor ExistingSeamColor = FColor::Green;
+	const FColor LockedColor = FColor::Green;
+	const FColor CompletionColor = FColor::Orange;
+	const FColor HoverColor = FColor::Green;
+
+	float ExistingSeamThickness = 2;
+	float ExistingSeamDepthBias = 0.1;
+
+	float LockedThickness = 3;
+	float LockedDepthBias = 1.0;
+
+	float HoverThickness = 3;
+	float HoverDepthBias = 1.1;
+
+	const float PointSize = 7;
+
+	void GetVidPath(FDynamicMesh3* Mesh, const TArray<int32>& StartVids, int32 EndVid, TArray<int32>& VidPathOut)
+	{
+		VidPathOut.Reset();
+		if (StartVids.Contains(EndVid))
+		{
+			VidPathOut.Add(EndVid);
+			return;
+		}
+
+		UE::Geometry::TMeshDijkstra<FDynamicMesh3> PathFinder(Mesh);
+		TArray<TMeshDijkstra<FDynamicMesh3>::FSeedPoint> SeedPoints;
+		for (int32 StartVid : StartVids)
+		{
+			SeedPoints.Add({ StartVid, StartVid, 0 });
+		}
+		
+		if (PathFinder.ComputeToTargetPoint(SeedPoints, EndVid))
+		{
+			PathFinder.FindPathToNearestSeed(EndVid, VidPathOut);
+		}
+		Algo::Reverse(VidPathOut);
+	}
+
+	// Like the other GetVidPath, just starts from a specific vert rather than from multiple
+	void GetVidPath(FDynamicMesh3* Mesh, int32 StartVid, int32 EndVid, TArray<int32>& VidPathOut)
+	{
+		GetVidPath(Mesh, TArray<int32> {StartVid}, EndVid, VidPathOut);
+	}
+
+	void AddDisplayedPoints(UUVEditorToolMeshInput* InputObject,
+		UPointSetComponent* UnwrapPointSet, UPointSetComponent* AppliedPointSet,
+		int32 AppliedVid, const TArray<int32>& UnwrapVids, const FColor& Color)
+	{
+		FTransform AppliedTransform = InputObject->AppliedPreview->PreviewMesh->GetTransform();
+		AppliedPointSet->AddPoint(FRenderablePoint(
+			InputObject->AppliedCanonical->GetVertex(AppliedVid), Color, PointSize));
+		for (int32 UnwrapVid : UnwrapVids)
+		{
+			UnwrapPointSet->AddPoint(FRenderablePoint(
+				InputObject->UnwrapCanonical->GetVertex(UnwrapVid), Color, PointSize));
+		}
+	}
+
+	void AddDisplayedPath(UUVEditorToolMeshInput* InputObject, 
+		ULineSetComponent* UnwrapLineSet, ULineSetComponent* AppliedLineSet,
+		const TArray<int32>& NewPathVids, bool bPathIsFromUnwrap, 
+		const FColor& Color, float Thickness, float DepthBias)
+	{
+		FTransform AppliedTransform = InputObject->AppliedPreview->PreviewMesh->GetTransform();
+		if (NewPathVids.Num() < 2)
+		{
+			return;
+		}
+
+		// Adds a path by simply walking through new path vids and using the given vid to position function.
+		auto AddSimplePath = [&NewPathVids, Color, Thickness, DepthBias](ULineSetComponent* LineSet,
+			TFunctionRef<FVector3d(int32)> VidToPosition)
+		{
+			FVector3d Vec1, Vec2;
+			FVector3d* EndPoint1 = &Vec1;
+			FVector3d* EndPoint2 = &Vec2;
+
+			*EndPoint1 = VidToPosition(NewPathVids[0]);
+			for (int32 i = 1; i < NewPathVids.Num(); ++i)
+			{
+				*EndPoint2 = VidToPosition(NewPathVids[i]);
+				LineSet->AddLine(FRenderableLine((FVector)*EndPoint1, (FVector)*EndPoint2,
+					Color, Thickness, DepthBias));
+				Swap(EndPoint1, EndPoint2);
+			}
+		};
+
+		// Add the path into the live preview.
+
+		TUniqueFunction<FVector3d(int32)> VidToAppliedMeshPosition = [InputObject, &AppliedTransform](int32 Vid) {
+			return AppliedTransform.TransformPosition(InputObject->AppliedCanonical->GetVertex(Vid));
+		};
+		if (bPathIsFromUnwrap)
+		{
+			VidToAppliedMeshPosition = [InputObject, &AppliedTransform](int32 Vid) {
+				return AppliedTransform.TransformPosition(
+					InputObject->AppliedCanonical->GetVertex(
+						InputObject->UnwrapVidToAppliedVid(Vid)));
+			};
+		}
+		AddSimplePath(AppliedLineSet, VidToAppliedMeshPosition);
+
+		// The work we have to do for the unwrap is actually very different depending on whether
+		// the path came from the unwrap or the applied mesh, because an applied mesh vertex may
+		// correspond to multiple unwrap vertices.
+		if (bPathIsFromUnwrap)
+		{
+			// The easy case, same as for live preview
+			TUniqueFunction<FVector3d(int32)> VidToUnwrapPosition = [InputObject](int32 Vid) {
+				return InputObject->UnwrapCanonical->GetVertex(Vid);
+			};
+			AddSimplePath(UnwrapLineSet, VidToUnwrapPosition);
+		}
+		else
+		{
+			// We'll only add an edge to the unwrap path if its unique. If it's not, the edge
+			// must be on a seam, and it's not useful to visualize it.
+			auto GetUniqueUnwrapEid = [InputObject](const TArray<int32>& Vids1, const TArray<int32>& Vids2)
+			{
+				int32 UniqueEid = IndexConstants::InvalidID;
+				for (int32 Vid1 : Vids1)
+				{
+					for (int32 Vid2 : Vids2)
+					{
+						int32 Eid = InputObject->UnwrapCanonical->FindEdge(Vid1, Vid2);
+						if (Eid != IndexConstants::InvalidID)
+						{
+							// Make sure this is the first edge we've found.
+							if (UniqueEid == IndexConstants::InvalidID)
+							{
+								UniqueEid = Eid;
+							}
+							else
+							{
+								// Not unique
+								return IndexConstants::InvalidID;
+							}
+						}
+					}
+				}
+				return UniqueEid;
+			};
+
+			TArray<int32> Vids1, Vids2;
+			TArray<int32>* EndPoints1 = &Vids1;
+			TArray<int32>* EndPoints2 = &Vids2;
+
+			InputObject->AppliedVidToUnwrapVids(NewPathVids[0], *EndPoints1);
+			for (int32 i = 1; i < NewPathVids.Num(); ++i)
+			{
+				InputObject->AppliedVidToUnwrapVids(NewPathVids[i], *EndPoints2);
+				int32 UnwrapEid = GetUniqueUnwrapEid(*EndPoints1, *EndPoints2);
+				if (UnwrapEid != IndexConstants::InvalidID)
+				{
+					FIndex2i EdgeVids = InputObject->UnwrapCanonical->GetEdgeV(UnwrapEid);
+					UnwrapLineSet->AddLine(FRenderableLine(
+						(FVector)InputObject->UnwrapCanonical->GetVertex(EdgeVids.A), 
+						(FVector)InputObject->UnwrapCanonical->GetVertex(EdgeVids.B),
+						Color, Thickness, DepthBias));
+				}
+				Swap(EndPoints1, EndPoints2);
+			}
+		}//end if unwrap path is not from unwrap
+	}//end AddDisplayedPath()
+
+	/**
+	 * Applies/reverts changes to LockedPath.
+	 */
+	class FPathChange : public FToolCommandChange
+	{
+	public:
+		/**
+		 * @param bIgnoreFirstPathVid If true, first vid in PathVidsIn is not considered. Useful
+		 *  because in cases where we're adding to an existing path, the first vert in the new
+		 *  path exists in the previous path.
+		 */
+		FPathChange(const TArray<int32> PathVidsIn, bool bRemovedIn, int32 MeshIndexIn, 
+			bool bIgnoreFirstPathVid = false)
+			: bRemoved(bRemovedIn)
+			, MeshIndex(MeshIndexIn)
+		{
+			for (int32 i = bIgnoreFirstPathVid ? 1 : 0; i < PathVidsIn.Num(); ++i)
+			{
+				PathVids.Add(PathVidsIn[i]);
+			}
+		}
+
+		virtual void Apply(UObject* Object) override
+		{
+			UUVEditorSeamTool* Tool = Cast<UUVEditorSeamTool>(Object);
+			if (bRemoved)
+			{
+				Tool->EditLockedPath([this](TArray<int32>& Path) {
+					if (ensure(Path.Num() >= PathVids.Num()))
+					{
+						Path.SetNum(Path.Num() - PathVids.Num());
+					}
+				}, MeshIndex);
+			}
+			else
+			{
+				Tool->EditLockedPath([this](TArray<int32>& Path) {
+					Path.Append(PathVids);
+				}, MeshIndex);
+			}
+		}
+
+		virtual void Revert(UObject* Object) override
+		{
+			UUVEditorSeamTool* Tool = Cast<UUVEditorSeamTool>(Object);
+			if (bRemoved)
+			{
+				Tool->EditLockedPath([this](TArray<int32>& Path) {
+					Path.Append(PathVids);
+				}, MeshIndex);
+			}
+			else
+			{
+				Tool->EditLockedPath([this](TArray<int32>& Path) {
+					if (ensure(Path.Num() >= PathVids.Num()))
+					{
+						Path.SetNum(Path.Num() - PathVids.Num());
+					}
+				}, MeshIndex);
+			}
+		}
+
+		virtual FString ToString() const override
+		{
+			return TEXT("UVEditorSeamToolLocals::FPathChange");
+		}
+
+	protected:
+		TArray<int32> PathVids;
+		bool bRemoved = false;
+		int32 MeshIndex = IndexConstants::InvalidID;
+	};
+}
+
+bool UUVEditorSeamToolBuilder::CanBuildTool(const FToolBuilderState& SceneState) const
+{
+	return Targets && Targets->Num() > 0;
+}
+
+UInteractiveTool* UUVEditorSeamToolBuilder::BuildTool(const FToolBuilderState& SceneState) const
+{
+	UUVEditorSeamTool* NewTool = NewObject<UUVEditorSeamTool>(SceneState.ToolManager);
+	NewTool->SetTargets(*Targets);
+
+	return NewTool;
+}
+
+void UUVEditorSeamTool::SetTargets(const TArray<TObjectPtr<UUVEditorToolMeshInput>>& TargetsIn)
+{
+	Targets = TargetsIn;
+}
+
+void UUVEditorSeamTool::Setup()
+{
+	using namespace UVEditorSeamToolLocals;
+
+	UInteractiveTool::Setup();
+
+	UContextObjectStore* ContextStore = GetToolManager()->GetContextObjectStore();
+	EmitChangeAPI = ContextStore->FindContext<UUVToolEmitChangeAPI>();
+	LivePreviewAPI = ContextStore->FindContext<UUVToolLivePreviewAPI>();
+	checkSlow(EmitChangeAPI && LivePreviewAPI);
+
+	// Set things up for being able to add behaviors to live preview.
+	LivePreviewBehaviorSet = NewObject<UInputBehaviorSet>();
+	LivePreviewBehaviorSource = NewObject<ULocalInputBehaviorSource>();
+	LivePreviewBehaviorSource->GetInputBehaviorsFunc = [this]() { return LivePreviewBehaviorSet; };
+
+	// Set up click and hover behaviors both for the 2d (unwrap) viewport and the 3d 
+	// (applied/live preview) viewport
+	ULocalSingleClickInputBehavior* UnwrapClickBehavior = NewObject<ULocalSingleClickInputBehavior>();
+	UnwrapClickBehavior->Initialize();
+	UnwrapClickBehavior->IsHitByClickFunc = [this](const FInputDeviceRay& InputRay) {
+		FInputRayHit HitResult;
+		HitResult.bHit = Get2DHitVertex(InputRay.WorldRay) != IndexConstants::InvalidID;
+		// We don't bother with the depth since there aren't competing click behaviors
+		return HitResult;
+	};
+	UnwrapClickBehavior->OnClickedFunc = [this](const FInputDeviceRay& ClickPos) {
+		int32 IndexOfMesh = IndexConstants::InvalidID;
+		int32 Vid = Get2DHitVertex(ClickPos.WorldRay, &IndexOfMesh);
+		if (Vid != IndexConstants::InvalidID)
+		{
+			OnMeshVertexClicked(Vid, IndexOfMesh, true);
+		}
+	};
+	AddInputBehavior(UnwrapClickBehavior);
+
+	ULocalMouseHoverBehavior* UnwrapHoverBehavior = NewObject<ULocalMouseHoverBehavior>();
+	UnwrapHoverBehavior->Initialize();
+	UnwrapHoverBehavior->BeginHitTestFunc = [this](const FInputDeviceRay& PressPos) {
+		FInputRayHit HitResult;
+		HitResult.bHit = Get2DHitVertex(PressPos.WorldRay) != IndexConstants::InvalidID;
+		// We don't bother with the depth since there aren't competing hover behaviors
+		return HitResult;
+	};
+	UnwrapHoverBehavior->OnUpdateHoverFunc = [this](const FInputDeviceRay& DevicePos)
+	{
+		int32 IndexOfMesh = IndexConstants::InvalidID;
+		int32 Vid = Get2DHitVertex(DevicePos.WorldRay, &IndexOfMesh);
+		if (Vid != IndexConstants::InvalidID)
+		{
+			OnMeshVertexHovered(Vid, IndexOfMesh, true);
+			return true;
+		}
+		return false;
+	};
+	UnwrapHoverBehavior->OnEndHoverFunc = [this]()
+	{
+		ClearHover();
+	};
+	AddInputBehavior(UnwrapHoverBehavior);
+
+	// Now the applied/live preview viewport...
+	ULocalSingleClickInputBehavior* LivePreviewClickBehavior = NewObject<ULocalSingleClickInputBehavior>();
+	LivePreviewClickBehavior->Initialize();
+	LivePreviewClickBehavior->IsHitByClickFunc = [this](const FInputDeviceRay& InputRay) {
+		FInputRayHit HitResult;
+		HitResult.bHit = Get3DHitVertex(InputRay.WorldRay) != IndexConstants::InvalidID;
+		// We don't bother with the depth since there aren't competing click behaviors
+		return HitResult;
+	};
+	LivePreviewClickBehavior->OnClickedFunc = [this](const FInputDeviceRay& ClickPos) {
+		int32 IndexOfMesh = IndexConstants::InvalidID;
+		int32 Vid = Get3DHitVertex(ClickPos.WorldRay, &IndexOfMesh);
+		if (Vid != IndexConstants::InvalidID)
+		{
+			OnMeshVertexClicked(Vid, IndexOfMesh, false);
+		}
+	};
+	LivePreviewBehaviorSet->Add(LivePreviewClickBehavior, this);
+
+	ULocalMouseHoverBehavior* LivePreviewHoverBehavior = NewObject<ULocalMouseHoverBehavior>();
+	LivePreviewHoverBehavior->Initialize();
+	LivePreviewHoverBehavior->BeginHitTestFunc = [this](const FInputDeviceRay& PressPos) {
+		FInputRayHit HitResult;
+		HitResult.bHit = Get3DHitVertex(PressPos.WorldRay) != IndexConstants::InvalidID;
+		// We don't bother with the depth since there aren't competing hover behaviors
+		return HitResult;
+	};
+	LivePreviewHoverBehavior->OnUpdateHoverFunc = [this](const FInputDeviceRay& DevicePos)
+	{
+		int32 IndexOfMesh = IndexConstants::InvalidID;
+		int32 Vid = Get3DHitVertex(DevicePos.WorldRay, &IndexOfMesh);
+		if (Vid != IndexConstants::InvalidID)
+		{
+			OnMeshVertexHovered(Vid, IndexOfMesh, false);
+			return true;
+		}
+		return false;
+	};
+	LivePreviewHoverBehavior->OnEndHoverFunc = [this]()
+	{
+		ClearHover();
+	};
+	LivePreviewBehaviorSet->Add(LivePreviewHoverBehavior, this);
+
+	// Give the live preview behaviors to the live preview input router
+	if (LivePreviewAPI)
+	{
+		LivePreviewInputRouter = LivePreviewAPI->GetLivePreviewInputRouter();
+		LivePreviewInputRouter->RegisterSource(LivePreviewBehaviorSource);
+	}
+
+	// Set up all the components we need to visualize things.
+	UnwrapGeometry = NewObject<UPreviewGeometry>();
+	UnwrapGeometry->CreateInWorld(Targets[0]->UnwrapPreview->GetWorld(), FTransform::Identity);
+	
+	LivePreviewGeometry = NewObject<UPreviewGeometry>();
+	LivePreviewGeometry->CreateInWorld(Targets[0]->AppliedPreview->GetWorld(),
+		Targets[0]->AppliedPreview->PreviewMesh->GetTransform());
+
+	// These visualize the locked-in portion of the current seam
+	UnwrapGeometry->AddPointSet(LockedPointSetID);
+	UnwrapGeometry->AddLineSet(LockedLineSetID);
+	LivePreviewGeometry->AddPointSet(LockedPointSetID);
+	LivePreviewGeometry->AddLineSet(LockedLineSetID);
+
+	// These visualize the start points of the locked-in seam. It's handy to have them
+	// separate so that they don't get cleared when the last endpoints get updated.
+	UnwrapGeometry->AddPointSet(StartPointSetID);
+	LivePreviewGeometry->AddPointSet(StartPointSetID);
+
+	// These visualize the currently hovered (preview) portion of the seam
+	UnwrapGeometry->AddPointSet(HoverPointSetID);
+	UnwrapGeometry->AddLineSet(HoverLineSetID);
+	LivePreviewGeometry->AddPointSet(HoverPointSetID);
+	LivePreviewGeometry->AddLineSet(HoverLineSetID);
+
+	// These visualize the existing seams on the 3d preview. These can get initialized now.
+	LivePreviewGeometry->AddLineSet(ExistingSeamsID);
+	ReconstructExistingSeamsVisualization();
+
+
+	UUVToolAABBTreeStorage* TreeStore = ContextStore->FindContext<UUVToolAABBTreeStorage>();
+	if (!TreeStore)
+	{
+		TreeStore = NewObject<UUVToolAABBTreeStorage>();
+		ContextStore->AddContextObject(TreeStore);
+	}
+
+	// Initialize the AABB trees from cached values, or make new ones.
+	for (int32 i = 0; i < Targets.Num(); ++i)
+	{
+		TSharedPtr<FDynamicMeshAABBTree3> Tree2d = TreeStore->Get(Targets[i]->UnwrapCanonical.Get());
+		if (!Tree2d)
+		{
+			Tree2d = MakeShared<FDynamicMeshAABBTree3>();
+			Tree2d->SetMesh(Targets[i]->UnwrapCanonical.Get());
+			TreeStore->Set(Targets[i]->UnwrapCanonical.Get(), Tree2d, Targets[i]);
+		}
+		Spatials2D.Add(Tree2d);
+
+		TSharedPtr<FDynamicMeshAABBTree3> Tree3d = TreeStore->Get(Targets[i]->AppliedCanonical.Get());
+		if (!Tree3d)
+		{
+			Tree3d = MakeShared<FDynamicMeshAABBTree3>();
+			Tree3d->SetMesh(Targets[i]->AppliedCanonical.Get());
+			TreeStore->Set(Targets[i]->AppliedCanonical.Get(), Tree3d, Targets[i]);
+		}
+		Spatials3D.Add(Tree3d);
+
+		// Make sure we update the seam visualization after mesh undo/redo
+		Targets[i]->OnCanonicalModified.AddWeakLambda(this, [this]
+		(UUVEditorToolMeshInput* InputObject, const UUVEditorToolMeshInput::FCanonicalModifiedInfo& Info) {
+			ReconstructExistingSeamsVisualization();
+		});
+	}
+}
+
+void UUVEditorSeamTool::ReconstructExistingSeamsVisualization()
+{
+	using namespace UVEditorSeamToolLocals;
+
+	ULineSetComponent* SeamLines = LivePreviewGeometry->FindLineSet(ExistingSeamsID);
+	SeamLines->Clear();
+	for (TObjectPtr<UUVEditorToolMeshInput> Target : Targets)
+	{
+		FDynamicMesh3* Mesh = Target->AppliedCanonical.Get();
+		FDynamicMeshUVOverlay* Overlay = Target->AppliedCanonical->Attributes()->GetUVLayer(Target->UVLayerIndex);
+		FTransform Transform = Target->AppliedPreview->PreviewMesh->GetTransform();
+		for (int32 Eid : Mesh->EdgeIndicesItr())
+		{
+			if (Overlay->IsSeamEdge(Eid))
+			{
+				FIndex2i EdgeVids = Mesh->GetEdgeV(Eid);
+				SeamLines->AddLine((FVector)Transform.TransformPosition(Mesh->GetVertex(EdgeVids.A)),
+					(FVector)Transform.TransformPosition(Mesh->GetVertex(EdgeVids.B)),
+					ExistingSeamColor, ExistingSeamThickness, ExistingSeamDepthBias);
+			}
+		}
+	}
+}
+
+void UUVEditorSeamTool::ReconstructLockedPathVisualization()
+{
+	using namespace UVEditorSeamToolLocals;
+
+	UPointSetComponent* UnwrapPointSet = UnwrapGeometry->FindPointSet(LockedPointSetID);
+	UPointSetComponent* AppliedPointSet = LivePreviewGeometry->FindPointSet(LockedPointSetID);
+	ULineSetComponent* UnwrapLineSet = UnwrapGeometry->FindLineSet(LockedLineSetID);
+	ULineSetComponent* AppliedLineSet = LivePreviewGeometry->FindLineSet(LockedLineSetID);
+	UPointSetComponent* UnwrapStartPoints = UnwrapGeometry->FindPointSet(StartPointSetID);
+	UPointSetComponent* AppliedStartPoints = LivePreviewGeometry->FindPointSet(StartPointSetID);
+
+	UnwrapLineSet->Clear();
+	AppliedLineSet->Clear();
+	UnwrapPointSet->Clear();
+	AppliedPointSet->Clear();
+	UnwrapStartPoints->Clear();
+	AppliedStartPoints->Clear();
+
+	if (ClickedMeshIndex != IndexConstants::InvalidID)
+	{
+		UUVEditorToolMeshInput* ClickedTarget = Targets[ClickedMeshIndex];
+
+		AddDisplayedPath(ClickedTarget, UnwrapLineSet, AppliedLineSet, LockedPath,
+			false, LockedColor, LockedThickness, LockedDepthBias);
+
+		TArray<int32> UnwrapVids;
+		ClickedTarget->AppliedVidToUnwrapVids(LastLockedAppliedVid, UnwrapVids);
+		AddDisplayedPoints(ClickedTarget, UnwrapPointSet, AppliedPointSet, LastLockedAppliedVid, UnwrapVids, LockedColor);
+
+		UnwrapVids.Reset();
+		ClickedTarget->AppliedVidToUnwrapVids(SeamStartAppliedVid, UnwrapVids);
+		AddDisplayedPoints(ClickedTarget, UnwrapStartPoints, AppliedStartPoints, SeamStartAppliedVid, UnwrapVids, LockedColor);
+	}
+}
+
+void UUVEditorSeamTool::Shutdown(EToolShutdownType ShutdownType)
+{
+	UInteractiveTool::Shutdown(ShutdownType);
+
+	EmitChangeAPI = nullptr;
+	LivePreviewAPI = nullptr;
+
+	if (LivePreviewInputRouter.IsValid())
+	{
+		LivePreviewInputRouter->DeregisterSource(LivePreviewBehaviorSource);
+	}
+	LivePreviewBehaviorSet->RemoveAll();
+
+	LivePreviewBehaviorSet = nullptr;
+	LivePreviewBehaviorSource = nullptr;
+
+	UnwrapGeometry->Disconnect();
+	LivePreviewGeometry->Disconnect();
+	UnwrapGeometry = nullptr;
+	LivePreviewGeometry = nullptr;
+
+	for (TObjectPtr<UUVEditorToolMeshInput> Target : Targets)
+	{
+		Target->OnCanonicalModified.RemoveAll(this);
+	}
+
+	Spatials2D.Reset();
+	Spatials3D.Reset();
+}
+
+void UUVEditorSeamTool::OnTick(float DeltaTime)
+{
+	using namespace UVEditorSeamToolLocals;
+
+	if (LivePreviewAPI)
+	{
+		LivePreviewAPI->GetLivePreviewCameraState(LivePreviewCameraState);
+	}
+
+	// This clears hover if necessary, etc. Includes checks to avoid redoing work, so
+	// no checks needed here.
+	UpdateHover();
+
+	// Apply click if there was one (signaled by a non-InvalidId vid)
+	if (ClickedVid != IndexConstants::InvalidID)
+	{
+		ApplyClick();
+		ClickedVid = IndexConstants::InvalidID;
+	}
+}
+
+void UUVEditorSeamTool::UpdateHover()
+{
+	using namespace UVEditorSeamToolLocals;
+
+	if (bLastHoverVidWasFromUnwrap == bHoverVidIsFromUnwrap
+		&& LastHoverMeshIndex == HoverMeshIndex
+		&& LastHoverVid == HoverVid)
+	{
+		return;
+	}
+
+	bLastHoverVidWasFromUnwrap = bHoverVidIsFromUnwrap;
+	LastHoverMeshIndex = HoverMeshIndex;
+	LastHoverVid = HoverVid;
+
+	ClearHover(false);
+
+	if (HoverVid == IndexConstants::InvalidID)
+	{
+		return;
+	}
+
+	int32 AppliedVid = bHoverVidIsFromUnwrap ? Targets[HoverMeshIndex]->UnwrapVidToAppliedVid(HoverVid) : HoverVid;
+	TArray<int32> UnwrapVids;
+	Targets[HoverMeshIndex]->AppliedVidToUnwrapVids(AppliedVid, UnwrapVids);
+
+	UPointSetComponent* UnwrapPointSet = UnwrapGeometry->FindPointSet(HoverPointSetID);
+	UPointSetComponent* AppliedPointSet = LivePreviewGeometry->FindPointSet(HoverPointSetID);
+
+	if (State == EState::WaitingToStart)
+	{
+		// Just draw the point and finish
+		AddDisplayedPoints(Targets[HoverMeshIndex], UnwrapPointSet, AppliedPointSet,
+			AppliedVid, UnwrapVids, HoverColor);
+		return;
+	}
+
+	// Otherwise, try to find a path.
+	TArray<int32> NewPathVids;
+	if (AppliedVid != LastLockedAppliedVid)
+	{
+		if (bHoverVidIsFromUnwrap)
+		{
+			TArray<int32> LastLockedUnwrapVids;
+			Targets[HoverMeshIndex]->AppliedVidToUnwrapVids(LastLockedAppliedVid, LastLockedUnwrapVids);
+			GetVidPath(Targets[HoverMeshIndex]->UnwrapCanonical.Get(), LastLockedUnwrapVids, HoverVid, NewPathVids);
+		}
+		else
+		{
+			ensure(LastLockedAppliedVid != IndexConstants::InvalidID);
+			GetVidPath(Targets[HoverMeshIndex]->AppliedCanonical.Get(), LastLockedAppliedVid, HoverVid, NewPathVids);
+		}
+	}
+
+	// Draw the path, if found.
+	if (NewPathVids.Num() > 0)
+	{
+		ULineSetComponent* UnwrapLineSet = UnwrapGeometry->FindLineSet(HoverLineSetID);
+		ULineSetComponent* AppliedLineSet = LivePreviewGeometry->FindLineSet(HoverLineSetID);
+		AddDisplayedPath(Targets[HoverMeshIndex], UnwrapLineSet, AppliedLineSet,
+			NewPathVids, bHoverVidIsFromUnwrap, HoverColor, HoverThickness, HoverDepthBias);
+	}
+
+	// See if we would apply the seam on click
+	if (AppliedVid == LastLockedAppliedVid || AppliedVid == SeamStartAppliedVid)
+	{
+		bCompletionColorOverride = true;
+		ResetPreviewColors();
+	}
+	else if (NewPathVids.Num() > 0)
+	{
+		// Otherwise draw the new point we're hovering. Don't do this if we didn't find a path
+		// there because that turns out to be confusing, esp when we have overlapping islands and
+		// the point is on another island.
+		AddDisplayedPoints(Targets[HoverMeshIndex], UnwrapPointSet, AppliedPointSet,
+			AppliedVid, UnwrapVids, HoverColor);
+	}
+}
+
+void UUVEditorSeamTool::ClearHover(bool bClearHoverInfo)
+{
+	using namespace UVEditorSeamToolLocals;
+
+	UnwrapGeometry->FindPointSet(HoverPointSetID)->Clear();
+	LivePreviewGeometry->FindPointSet(HoverPointSetID)->Clear();
+	UnwrapGeometry->FindLineSet(HoverLineSetID)->Clear();
+	LivePreviewGeometry->FindLineSet(HoverLineSetID)->Clear();
+
+	if (bClearHoverInfo)
+	{
+		HoverVid = IndexConstants::InvalidID;
+		HoverMeshIndex = IndexConstants::InvalidID;
+		LastHoverVid = HoverVid;
+		LastHoverMeshIndex = HoverMeshIndex;
+	}
+
+	if (bCompletionColorOverride)
+	{
+		bCompletionColorOverride = false;
+		ResetPreviewColors();
+	}
+}
+
+void UUVEditorSeamTool::ResetPreviewColors()
+{
+	using namespace UVEditorSeamToolLocals;
+
+	FColor Color = bCompletionColorOverride ? CompletionColor : HoverColor;
+
+	UnwrapGeometry->FindPointSet(HoverPointSetID)->SetAllPointsColor(Color);
+	LivePreviewGeometry->FindPointSet(HoverPointSetID)->SetAllPointsColor(Color);
+	UnwrapGeometry->FindLineSet(HoverLineSetID)->SetAllLinesColor(Color);
+	LivePreviewGeometry->FindLineSet(HoverLineSetID)->SetAllLinesColor(Color);
+
+	Color = bCompletionColorOverride ? CompletionColor : LockedColor;
+
+	UnwrapGeometry->FindPointSet(LockedPointSetID)->SetAllPointsColor(Color);
+	LivePreviewGeometry->FindPointSet(LockedPointSetID)->SetAllPointsColor(Color);
+	UnwrapGeometry->FindLineSet(LockedLineSetID)->SetAllLinesColor(Color);
+	LivePreviewGeometry->FindLineSet(LockedLineSetID)->SetAllLinesColor(Color);
+	UnwrapGeometry->FindPointSet(StartPointSetID)->SetAllPointsColor(Color);
+	LivePreviewGeometry->FindPointSet(StartPointSetID)->SetAllPointsColor(Color);
+}
+
+void UUVEditorSeamTool::ApplyClick()
+{
+	using namespace UVEditorSeamToolLocals;
+
+	ClearHover();
+
+	UUVEditorToolMeshInput* ClickedTarget = Targets[ClickedMeshIndex];
+	FTransform AppliedTransform = ClickedTarget->AppliedPreview->PreviewMesh->GetTransform();
+	
+	int32 AppliedVid = bClickWasInUnwrap ? ClickedTarget->UnwrapVidToAppliedVid(ClickedVid) : ClickedVid;
+	TArray<int32> UnwrapVids;
+	ClickedTarget->AppliedVidToUnwrapVids(AppliedVid, UnwrapVids);
+
+	if (State == EState::WaitingToStart)
+	{
+		SeamStartAppliedVid = AppliedVid;
+		LastLockedAppliedVid = SeamStartAppliedVid;
+
+		UPointSetComponent* UnwrapPointSet = UnwrapGeometry->FindPointSet(StartPointSetID);
+		UPointSetComponent* AppliedPointSet = LivePreviewGeometry->FindPointSet(StartPointSetID);
+		AddDisplayedPoints(ClickedTarget, UnwrapPointSet, AppliedPointSet, AppliedVid, UnwrapVids, LockedColor);
+
+		LockedPath.Add(AppliedVid);
+		EmitChangeAPI->EmitToolDependentChange(this, 
+			MakeUnique<FPathChange>(LockedPath, false, ClickedMeshIndex),
+			EditSeamTransactionName);
+
+		State = EState::SeamInProgress;
+		return;
+	}
+
+	// Otherwise, we're working on a seam...
+	// Find a path for the new section
+	TArray<int32> VidPath;
+	TArray<int32>* AppliedVidPath = &VidPath;
+	TArray<int32> TempStorage; // Only used if we need spare storage for AppliedVids
+	if (bClickWasInUnwrap)
+	{
+		TArray<int32> LastLockedUnwrapVids;
+		ClickedTarget->AppliedVidToUnwrapVids(LastLockedAppliedVid, LastLockedUnwrapVids);
+		GetVidPath(ClickedTarget->UnwrapCanonical.Get(), LastLockedUnwrapVids, ClickedVid, VidPath);
+		for (int32 Vid : VidPath)
+		{
+			TempStorage.Add(ClickedTarget->UnwrapVidToAppliedVid(Vid));
+		}
+		AppliedVidPath = &TempStorage;
+	}
+	else
+	{
+		GetVidPath(ClickedTarget->AppliedCanonical.Get(), LastLockedAppliedVid, ClickedVid, VidPath);
+	}
+
+	if (VidPath.Num() == 0)
+	{
+		// TODO: Write an error message complaining that path couldn't be found?
+		return;
+	}
+
+	// See if we completed a seam
+	if (AppliedVid == LastLockedAppliedVid || AppliedVid == SeamStartAppliedVid)
+	{
+		// We keep LockedPath unchanged and create a separate path copy to pass to ApplyPath
+		// so that when we clear LockedPath, the undo transaction would bring us back to the
+		// pre-completed state.
+		TArray<int32> SeamVids = LockedPath;
+		if (ensure(SeamVids.Num() > 0)) {
+			// Avoid duplicating the last point
+			SeamVids.Pop();
+		}
+		SeamVids.Append(*AppliedVidPath);
+
+		ApplySeam(SeamVids);
+		return;
+	}
+
+	// If not completed, need to update path, visualization, and emit change
+	if (ensure(LockedPath.Num() > 0)) {
+		// Avoid duplicating the last point
+		LockedPath.Pop();
+	}
+	LockedPath.Append(*AppliedVidPath);
+
+	UPointSetComponent* UnwrapPointSet = UnwrapGeometry->FindPointSet(LockedPointSetID);
+	UPointSetComponent* AppliedPointSet = LivePreviewGeometry->FindPointSet(LockedPointSetID);
+	ULineSetComponent* UnwrapLineSet = UnwrapGeometry->FindLineSet(LockedLineSetID);
+	ULineSetComponent* AppliedLineSet = LivePreviewGeometry->FindLineSet(LockedLineSetID);
+	AddDisplayedPath(ClickedTarget, UnwrapLineSet, AppliedLineSet, VidPath,
+		bClickWasInUnwrap, LockedColor, LockedThickness, LockedDepthBias);
+	UnwrapPointSet->Clear();
+	AppliedPointSet->Clear();
+	AddDisplayedPoints(ClickedTarget, UnwrapPointSet, AppliedPointSet, AppliedVid, UnwrapVids, LockedColor);
+	
+	EmitChangeAPI->EmitToolDependentChange(this, 
+		MakeUnique<FPathChange>(*AppliedVidPath, false, ClickedMeshIndex, true),
+		EditSeamTransactionName);
+
+	LastLockedAppliedVid = AppliedVid;
+}
+
+void UUVEditorSeamTool::ApplySeam(const TArray<int32>& AppliedVidsIn)
+{
+	using namespace UVEditorSeamToolLocals;
+
+	UUVEditorToolMeshInput* ClickedTarget = Targets[ClickedMeshIndex];
+
+	// The locked path becomes an actual UV seam.
+
+	// Convert vid path into set of eids path in the applied mesh
+	TSet<int32> EidSet;
+	for (int32 i = 0; i < AppliedVidsIn.Num() - 1; ++i)
+	{
+		int32 Eid = ClickedTarget->AppliedCanonical->FindEdge(AppliedVidsIn[i], AppliedVidsIn[i + 1]);
+		if (ensure(Eid != IndexConstants::InvalidID))
+		{
+			EidSet.Add(ClickedTarget->AppliedCanonical->FindEdge(AppliedVidsIn[i], AppliedVidsIn[i + 1]));
+		}
+	}
+
+	FUVEditResult UVEditResult;
+	FDynamicMeshUVEditor UVEditor(ClickedTarget->AppliedCanonical.Get(), 
+		ClickedTarget->UVLayerIndex, false);
+	UVEditor.CreateSeamsAtEdges(EidSet, &UVEditResult);
+
+	TSet<int32> TidSet;
+	for (int32 UnwrapVid : UVEditResult.NewUVElements)
+	{
+		TArray<int32> VertTids;
+		ClickedTarget->AppliedCanonical->GetVtxTriangles(ClickedTarget->UnwrapVidToAppliedVid(UnwrapVid), VertTids);
+		TidSet.Append(VertTids);
+	}
+
+	FDynamicMeshChangeTracker ChangeTracker(ClickedTarget->UnwrapCanonical.Get());
+	ChangeTracker.BeginChange();
+	ChangeTracker.SaveTriangles(TidSet, true);
+
+	TArray<int32> AppliedTids = TidSet.Array();
+	ClickedTarget->UpdateAllFromAppliedCanonical(&UVEditResult.NewUVElements, &AppliedTids, &AppliedTids);
+	Spatials2D[ClickedMeshIndex]->Build();
+
+	// Emit transaction
+	EmitChangeAPI->BeginUndoTransaction(ApplySeamTransactionName);
+	 // emits locked path change
+	ClearLockedPath();
+	EmitChangeAPI->EmitToolIndependentUnwrapCanonicalChange(
+		ClickedTarget, ChangeTracker.EndChange(), EditSeamTransactionName);
+	EmitChangeAPI->EndUndoTransaction();
+
+	ReconstructExistingSeamsVisualization();
+}
+
+void UUVEditorSeamTool::ClearLockedPath(bool bEmitChange)
+{
+	using namespace UVEditorSeamToolLocals;
+
+	if (bEmitChange)
+	{
+		EmitChangeAPI->EmitToolDependentChange(this,
+			MakeUnique<FPathChange>(LockedPath, true, ClickedMeshIndex),
+			EditSeamTransactionName);
+	}
+	LockedPath.Reset();
+
+	UnwrapGeometry->FindPointSet(LockedPointSetID)->Clear();
+	LivePreviewGeometry->FindPointSet(LockedPointSetID)->Clear();
+	UnwrapGeometry->FindLineSet(LockedLineSetID)->Clear();
+	LivePreviewGeometry->FindLineSet(LockedLineSetID)->Clear();
+	UnwrapGeometry->FindPointSet(StartPointSetID)->Clear();
+	LivePreviewGeometry->FindPointSet(StartPointSetID)->Clear();
+
+	ClearHover();
+
+	ClickedMeshIndex = IndexConstants::InvalidID;
+	ClickedVid = IndexConstants::InvalidID;
+	SeamStartAppliedVid = IndexConstants::InvalidID;
+	LastLockedAppliedVid = IndexConstants::InvalidID;
+
+	State = EState::WaitingToStart;
+}
+
+int32 UUVEditorSeamTool::Get2DHitVertex(const FRay& WorldRayIn, int32* IndexOf2DSpatialOut)
+{
+	if (!ensure(WorldRayIn.Direction.Z != 0))
+	{
+		return IndexConstants::InvalidID;
+	}
+
+	// Find plane hit location
+	double PlaneHitT = WorldRayIn.Origin.Z / -WorldRayIn.Direction.Z;
+	FVector3d PlanePoint = WorldRayIn.PointAt(PlaneHitT);
+
+	int32 Result = IndexConstants::InvalidID;
+	if (State == EState::WaitingToStart)
+	{
+		// For tolerance, we'd like to divide the screen into 90 sections to be similar to the 90 degree
+		// FOV angle tolerance we use in 3D. Ray Z should be our distance from ground plane, so width
+		// would be 2Z for 90 degree FOV.
+		double Tolerance = ToolSceneQueriesUtil::GetDefaultVisualAngleSnapThreshD() * WorldRayIn.Origin.Z / 45;
+
+		double MinDistSquared = TNumericLimits<double>::Max();
+		for (int32 i = 0; i < Spatials2D.Num(); ++i)
+		{
+			double DistSquared = TNumericLimits<double>::Max();
+			int32 Vid = Spatials2D[i]->FindNearestVertex(PlanePoint, DistSquared, Tolerance);
+
+			if (DistSquared >= 0 && DistSquared < MinDistSquared)
+			{
+				MinDistSquared = DistSquared;
+				Result = Vid;
+				if (IndexOf2DSpatialOut)
+				{
+					*IndexOf2DSpatialOut = i;
+				}
+			}
+		}
+	}
+	else if (State == EState::SeamInProgress)
+	{
+		// We don't provide a tolerance here because we want to snap to nearest vertex even if 
+		// we're far outside the mesh
+		double Distance = -1;
+		int32 Vid = Spatials2D[ClickedMeshIndex]->FindNearestVertex(PlanePoint, Distance);
+		if (Distance >= 0)
+		{
+			Result = Vid;
+			if (IndexOf2DSpatialOut)
+			{
+				*IndexOf2DSpatialOut = ClickedMeshIndex;
+			}
+		}
+	}
+
+	return Result;
+}
+
+int32 UUVEditorSeamTool::Get3DHitVertex(const FRay& WorldRayIn, int32* IndexOf3DSpatialOut)
+{
+	// When we hit a triangle, gives the closest vert in terms of angle to camera (kind of). Tests
+	// the three verts of the triangle we hit, and the closest vert to the hit location (to make it
+	// easier to hit verts in very thin/degenerate tris).
+	auto GetBestSnappedVert = [this](FDynamicMeshAABBTree3& Tree, const FTransform& MeshTransform,
+		const FVector3d& HitPoint, int32 Tid, double& BestSnapMetricOut)
+	{
+		const FDynamicMesh3* Mesh = Tree.GetMesh();
+		FIndex3i TriVids = Mesh->GetTriangle(Tid);
+
+		double Unused = 0; // dummy variable for function call
+		int32 VidCandidates[4] = { TriVids.A, TriVids.B, TriVids.C, Tree.FindNearestVertex(HitPoint, Unused) };
+		
+		FVector3d VertCandidates[4];
+		Mesh->GetTriVertices(Tid, VertCandidates[0], VertCandidates[1], VertCandidates[2]);
+		VertCandidates[3] = Mesh->GetVertex(VidCandidates[3]);
+
+		BestSnapMetricOut = TNumericLimits<double>::Max();
+		int32 BestVid = IndexConstants::InvalidID;
+		for (int i = 0; i < 3; ++i)
+		{
+			double SnapMetric = ToolSceneQueriesUtil::PointSnapMetric(
+				LivePreviewCameraState, HitPoint, MeshTransform.TransformPosition(VertCandidates[i]));
+
+			if (SnapMetric < BestSnapMetricOut)
+			{
+				BestSnapMetricOut = SnapMetric;
+				BestVid = VidCandidates[i];
+			}
+		}
+
+		return BestVid;
+	};
+
+
+	if (State == EState::WaitingToStart)
+	{
+		// We'll hit test all the preview meshes
+		int32 IndexOfBestSpatial = IndexConstants::InvalidID;
+		double BestRayT = TNumericLimits<double>::Max();
+		int32 BestTid = IndexConstants::InvalidID;
+		for (int32 i = 0; i < Spatials3D.Num(); ++i)
+		{
+			FTransform Transform = Targets[i]->AppliedPreview->PreviewMesh->GetTransform();
+			FRay LocalRay(
+				Transform.InverseTransformPosition(WorldRayIn.Origin),
+				Transform.InverseTransformVector(WorldRayIn.Direction));
+
+			int32 Tid = IndexConstants::InvalidID;
+			double RayT = TNumericLimits<double>::Max();
+			if (Spatials3D[i]->FindNearestHitTriangle(LocalRay, RayT, Tid))
+			{
+				if (RayT >= 0 && RayT < BestRayT)
+				{
+					BestRayT = RayT;
+					BestTid = Tid;
+					IndexOfBestSpatial = i;
+				}
+			}
+		}
+
+		if (BestTid == IndexConstants::InvalidID)
+		{
+			// Didn't hit anything
+			return IndexConstants::InvalidID;
+		}
+
+		// Find vertex that's within some angle tolerance
+		double BestSnapMetric = TNumericLimits<double>::Max();
+		FTransform Transform = Targets[IndexOfBestSpatial]->AppliedPreview->PreviewMesh->GetTransform();
+		int32 BestVid = GetBestSnappedVert(*Spatials3D[IndexOfBestSpatial],
+			Transform, WorldRayIn.PointAt(BestRayT), BestTid, BestSnapMetric);
+
+		if (BestSnapMetric <= ToolSceneQueriesUtil::GetDefaultVisualAngleSnapThreshD())
+		{
+			if (IndexOf3DSpatialOut)
+			{
+				*IndexOf3DSpatialOut = IndexOfBestSpatial;
+			}
+			return BestVid;
+		}
+		return IndexConstants::InvalidID;
+	}
+	else if (State == EState::SeamInProgress)
+	{
+		// Only need to hit test the mesh we've already started on, and we always snap to nearest vert
+		FTransform Transform = Targets[ClickedMeshIndex]->AppliedPreview->PreviewMesh->GetTransform();
+		FRay LocalRay(
+			Transform.InverseTransformPosition(WorldRayIn.Origin),
+			Transform.InverseTransformVector(WorldRayIn.Direction));
+
+		double RayT = TNumericLimits<double>::Max();
+		int32 Tid = IndexConstants::InvalidID;
+		if (Spatials3D[ClickedMeshIndex]->FindNearestHitTriangle(LocalRay, RayT, Tid))
+		{
+			double BestSnapMetric = TNumericLimits<double>::Max();
+			if (IndexOf3DSpatialOut)
+			{
+				*IndexOf3DSpatialOut = ClickedMeshIndex;
+			}
+			return GetBestSnappedVert(*Spatials3D[ClickedMeshIndex],
+				Transform, WorldRayIn.PointAt(RayT), Tid, BestSnapMetric);
+		}
+		return IndexConstants::InvalidID;
+	}
+
+	// Shouldn't get here: all enum cases should have returns above
+	ensure(false);
+	return IndexConstants::InvalidID;
+}
+
+void UUVEditorSeamTool::OnMeshVertexClicked(int32 Vid, int32 IndexOfMesh, bool bVidIsFromUnwrap)
+{
+	using namespace UVEditorSeamToolLocals;
+
+	if (ClickedVid != IndexConstants::InvalidID)
+	{
+		// Haven't yet finished processing the last click
+		return;
+	}
+
+	// Save the click information. The click will be applied on tick.
+	ClickedVid = Vid;
+	ClickedMeshIndex = IndexOfMesh;
+	bClickWasInUnwrap = bVidIsFromUnwrap;
+}
+
+void UUVEditorSeamTool::OnMeshVertexHovered(int32 Vid, int32 IndexOfMesh, bool bVidIsFromUnwrap)
+{
+	using namespace UVEditorSeamToolLocals;
+
+	// Save the hover information. Hover will be updated on tick
+	HoverVid = Vid;
+	HoverMeshIndex = IndexOfMesh;
+	bHoverVidIsFromUnwrap = bVidIsFromUnwrap;
+}
+
+void UUVEditorSeamTool::EditLockedPath(
+	TUniqueFunction<void(TArray<int32>& LockedPathInOut)> EditFunction, int32 MeshIndex)
+{
+	EditFunction(LockedPath);
+
+	bool bPathNotEmpty = LockedPath.Num() > 0;
+	State = bPathNotEmpty ? EState::SeamInProgress : EState::WaitingToStart;
+	ClickedMeshIndex = bPathNotEmpty ? MeshIndex : IndexConstants::InvalidID;
+	LastLockedAppliedVid = bPathNotEmpty ? LockedPath.Last() : IndexConstants::InvalidID;
+	SeamStartAppliedVid = bPathNotEmpty ? LockedPath[0] : IndexConstants::InvalidID;
+
+	ReconstructLockedPathVisualization();
+	ClearHover();
+}
+
+bool UUVEditorSeamTool::CanCurrentlyNestedCancel()
+{
+	return LockedPath.Num() > 0;
+}
+
+bool UUVEditorSeamTool::ExecuteNestedCancelCommand()
+{
+	if (LockedPath.Num() > 0)
+	{
+		ClearLockedPath();
+		return true;
+	}
+	return false;
+}
+
+bool UUVEditorSeamTool::CanCurrentlyNestedAccept()
+{
+	return LockedPath.Num() > 0;
+}
+
+bool UUVEditorSeamTool::ExecuteNestedAcceptCommand()
+{
+	if (LockedPath.Num() > 0)
+	{
+		ApplySeam(LockedPath);
+		return true;
+	}
+	return false;
+}
+
+#undef LOCTEXT_NAMESPACE
diff --git a/Engine/Plugins/Experimental/UVEditor/Source/UVEditorTools/Public/UVEditorLayerEditTool.h b/Engine/Plugins/Experimental/UVEditor/Source/UVEditorTools/Public/UVEditorLayerEditTool.h
--- a/Engine/Plugins/Experimental/UVEditor/Source/UVEditorTools/Public/UVEditorLayerEditTool.h
+++ b/Engine/Plugins/Experimental/UVEditor/Source/UVEditorTools/Public/UVEditorLayerEditTool.h
@@ -184,9 +184,9 @@ public:
 
 	virtual void OnTick(float DeltaTime) override;
 
-	virtual bool HasCancel() const override { return true; }
+	virtual bool HasCancel() const override { return false; }
 	virtual bool HasAccept() const override { return false; }
-	virtual bool CanAccept() const override;
+	virtual bool CanAccept() const override { return false; }
 
 	virtual void OnPropertyModified(UObject* PropertySet, FProperty* Property) override;
 
diff --git a/Engine/Plugins/Experimental/UVEditor/Source/UVEditorTools/Public/UVEditorSeamTool.h b/Engine/Plugins/Experimental/UVEditor/Source/UVEditorTools/Public/UVEditorSeamTool.h
new file mode 100644
--- /dev/null
+++ b/Engine/Plugins/Experimental/UVEditor/Source/UVEditorTools/Public/UVEditorSeamTool.h
@@ -0,0 +1,146 @@
+// Copyright Epic Games, Inc. All Rights Reserved.
+
+#pragma once
+
+#include "CoreMinimal.h"
+
+#include "DynamicMesh/DynamicMeshAABBTree3.h"
+#include "IndexTypes.h"
+#include "InteractiveTool.h"
+#include "InteractiveToolBuilder.h"
+#include "InteractiveToolQueryInterfaces.h" // IInteractiveToolNestedAcceptCancelAPI
+#include "GeometryBase.h"
+
+#include "UVEditorSeamTool.generated.h"
+
+//class ULocalSingleClickInputBehavior;
+class UInputRouter;
+class ULocalInputBehaviorSource;
+class UPreviewGeometry;
+class UUVEditorToolMeshInput;
+class UUVToolEmitChangeAPI;
+class UUVToolLivePreviewAPI;
+
+PREDECLARE_GEOMETRY(class FDynamicMesh);
+
+UCLASS()
+class UVEDITORTOOLS_API UUVEditorSeamToolBuilder : public UInteractiveToolBuilder
+{
+	GENERATED_BODY()
+public:
+
+	virtual bool CanBuildTool(const FToolBuilderState& SceneState) const override;
+	virtual UInteractiveTool* BuildTool(const FToolBuilderState& SceneState) const override;
+
+	// These are pointers so that they can be updated under the builder without
+	// having to reset them after things are reinitialized.
+	const TArray<TObjectPtr<UUVEditorToolMeshInput>>* Targets = nullptr;
+};
+
+UCLASS()
+class UVEDITORTOOLS_API UUVEditorSeamTool : public UInteractiveTool,
+	public IInteractiveToolNestedAcceptCancelAPI
+{
+	GENERATED_BODY()
+
+public:
+	using FDynamicMeshAABBTree3 = UE::Geometry::FDynamicMeshAABBTree3;
+
+	virtual void SetTargets(const TArray<TObjectPtr<UUVEditorToolMeshInput>>& TargetsIn);
+
+	// For use by undo/redo
+	void EditLockedPath(
+		TUniqueFunction<void(TArray<int32>& LockedPathInOut)> EditFunction, int32 MeshIndex);
+
+	// UInteractiveTool
+	virtual void Setup() override;
+	virtual void Shutdown(EToolShutdownType ShutdownType) override;
+	virtual void OnTick(float DeltaTime) override;
+	virtual bool HasCancel() const override { return false; }
+	virtual bool HasAccept() const override { return false; }
+
+	// IInteractiveToolNestedAcceptCancelAPI
+	virtual bool SupportsNestedCancelCommand() override { return true; }
+	virtual bool CanCurrentlyNestedCancel() override;
+	virtual bool ExecuteNestedCancelCommand() override;
+	virtual bool SupportsNestedAcceptCommand() override { return true; }
+	virtual bool CanCurrentlyNestedAccept() override;
+	virtual bool ExecuteNestedAcceptCommand() override;
+protected:
+	void ReconstructExistingSeamsVisualization();
+	void ReconstructLockedPathVisualization();
+
+	int32 Get2DHitVertex(const FRay& WorldRayIn, int32* IndexOf2DSpatialOut = nullptr);
+	int32 Get3DHitVertex(const FRay& WorldRayIn, int32* IndexOf3DSpatialOut = nullptr);
+
+	void OnMeshVertexClicked(int32 Vid, int32 IndexOfMesh, bool bVidIsFromUnwrap);
+	void OnMeshVertexHovered(int32 Vid, int32 IndexOfMesh, bool bVidIsFromUnwrap);
+	void UpdateHover();
+	/** @param bClearHoverInfo If true, also clears HoverVid, etc in addition to just clearing display. */
+	void ClearHover(bool bClearHoverInfo = true);
+	void ResetPreviewColors();
+	void ApplyClick();
+	/** Clears LockedAppliedVids, but builds seam off of AppliedVidsIn */
+	void ApplySeam(const TArray<int32>& AppliedVidsIn);
+	void ClearLockedPath(bool bEmitChange = true);
+
+	enum class EState
+	{
+		WaitingToStart,
+		SeamInProgress
+	};
+
+	EState State = EState::WaitingToStart;
+
+	UPROPERTY()
+	TArray<TObjectPtr<UUVEditorToolMeshInput>> Targets;
+
+	UPROPERTY()
+	TObjectPtr<UUVToolLivePreviewAPI> LivePreviewAPI = nullptr;
+
+	UPROPERTY()
+	TObjectPtr<UUVToolEmitChangeAPI> EmitChangeAPI = nullptr;
+
+	UPROPERTY()
+	TObjectPtr<UInputBehaviorSet> LivePreviewBehaviorSet = nullptr;
+
+	UPROPERTY()
+	TObjectPtr<ULocalInputBehaviorSource> LivePreviewBehaviorSource = nullptr;
+
+	UPROPERTY()
+	TObjectPtr<UPreviewGeometry> UnwrapGeometry = nullptr;
+
+	UPROPERTY()
+	TObjectPtr<UPreviewGeometry> LivePreviewGeometry = nullptr;
+
+	TWeakObjectPtr<UInputRouter> LivePreviewInputRouter = nullptr;
+
+	TArray<TSharedPtr<FDynamicMeshAABBTree3>> Spatials2D; // 1:1 with targets
+	TArray<TSharedPtr<FDynamicMeshAABBTree3>> Spatials3D; // 1:1 with targets
+
+	FViewCameraState LivePreviewCameraState;
+
+	// Used to remember click info to apply on tick
+	int32 ClickedVid = IndexConstants::InvalidID;
+	int32 ClickedMeshIndex = -1;
+	bool bClickWasInUnwrap = false;
+
+	// Used to remember hover info to apply on tick
+	int32 HoverVid = IndexConstants::InvalidID;
+	int32 HoverMeshIndex = IndexConstants::InvalidID;
+	bool bHoverVidIsFromUnwrap = false;
+	int32 LastHoverVid = IndexConstants::InvalidID;
+	int32 LastHoverMeshIndex = IndexConstants::InvalidID;
+	bool bLastHoverVidWasFromUnwrap = false;
+
+	// Used to know when to end the seam.
+	int32 SeamStartAppliedVid = IndexConstants::InvalidID;
+	int32 LastLockedAppliedVid = IndexConstants::InvalidID;
+
+	TArray<int32> LockedPath;
+	TArray<int32> PreviewPath;
+
+	// When true, the entire path is changed to the "completion" color to show
+	// that the next click will complete the path.
+	bool bCompletionColorOverride = false;
+};
