UV Editor: Move background and UV channel settings into the toolbar at the top. Also move UVEditorBackgroundPreview.h to be in the proper folder, remove "force rebuild" parameter in visibility API, and fix various undo/redo bugs in UVEditorLayerEditTool (transactions not being expired, delete undo causing crashes).

#rnx

diff --git a/Engine/Plugins/Experimental/UVEditor/Source/UVEditor/Private/UVEditorBackgroundPreview.h b/Engine/Plugins/Experimental/UVEditor/Source/UVEditor/Private/UVEditorBackgroundPreview.h
deleted file mode 100644
diff --git a/Engine/Plugins/Experimental/UVEditor/Source/UVEditor/Private/UVEditorMode.cpp b/Engine/Plugins/Experimental/UVEditor/Source/UVEditor/Private/UVEditorMode.cpp
--- a/Engine/Plugins/Experimental/UVEditor/Source/UVEditor/Private/UVEditorMode.cpp
+++ b/Engine/Plugins/Experimental/UVEditor/Source/UVEditor/Private/UVEditorMode.cpp
@@ -32,7 +32,6 @@
 #include "UVEditorToolUtil.h"
 #include "UVToolContextObjects.h"
 #include "UVEditorBackgroundPreview.h"
-#include "UVEditorModeChannelProperties.h"
 #include "Editor.h"
 
 #define LOCTEXT_NAMESPACE "UUVEditorMode"
@@ -70,14 +69,12 @@ namespace UVEditorModeLocals
 		{
 			UUVEditorMode* UVEditorMode = Cast<UUVEditorMode>(Object);
 			UVEditorMode->ChangeInputObjectLayer(AssetID, NewUVLayerIndex);
-			UVEditorMode->UpdateSelectedLayer();
 		}
 
 		virtual void Revert(UObject* Object) override
 		{
 			UUVEditorMode* UVEditorMode = Cast<UUVEditorMode>(Object);
 			UVEditorMode->ChangeInputObjectLayer(AssetID, OldUVLayerIndex);
-			UVEditorMode->UpdateSelectedLayer();
 		}
 
 		virtual bool HasExpired(UObject* Object) const override
@@ -97,8 +94,15 @@ namespace UVEditorModeLocals
 		int32 NewUVLayerIndex;
 	};
 
+	void InitializeAssetNames(const TArray<TObjectPtr<UToolTarget>>& ToolTargets, TArray<FString>& AssetNamesOut)
+	{
+		AssetNamesOut.Reset(ToolTargets.Num());
 
-
+		for (const TObjectPtr<UToolTarget>& ToolTarget : ToolTargets)
+		{
+			AssetNamesOut.Add(UE::ToolTarget::GetHumanReadableName(ToolTarget));
+		}
+	}
 }
 
 const FToolTargetTypeRequirements& UUVEditorMode::GetToolTargetRequirements()
@@ -132,17 +136,10 @@ void UUVEditorMode::Enter()
 	BackgroundVisualization = NewObject<UUVEditorBackgroundPreview>(this);
 	BackgroundVisualization->CreateInWorld(GetWorld(), FTransform::Identity);
 
-	UVChannelProperties = NewObject<UUVEditorUVChannelProperties>(this);
-	UVChannelProperties->WatchProperty(UVChannelProperties->Asset, [this](FString UVAsset) {SwitchActiveAsset(UVAsset); });
-	UVChannelProperties->WatchProperty(UVChannelProperties->UVChannel, [this](FString UVChannel) {SwitchActiveChannel(UVChannel); });
-
 	BackgroundVisualization->Settings->WatchProperty(BackgroundVisualization->Settings->bVisible, 
 		[this](bool IsVisible) { UpdateTriangleMaterialBasedOnBackground(IsVisible); });
 
-	AddDisplayedPropertySet(UVChannelProperties);
-	AddDisplayedPropertySet(BackgroundVisualization->Settings);
-		
-	StaticCastSharedPtr<FUVEditorModeToolkit>(Toolkit)->SetModeDetailsViewObjects(PropertyObjectsToDisplay);	
+	PropertyObjectsToTick.Add(BackgroundVisualization->Settings);
 	
 	RegisterTools();
 
@@ -186,15 +183,18 @@ void UUVEditorMode::RegisterTools()
 	RegisterTool(CommandInfos.BeginSeamTool, TEXT("UVSeamTool"), UVEditorSeamToolBuilder);
 }
 
-void UUVEditorMode::AddDisplayedPropertySet(const TObjectPtr<UInteractiveToolPropertySet>& PropertySet)
+void UUVEditorMode::CreateToolkit()
 {
-	PropertyObjectsToTick.Add(PropertySet);
-	PropertyObjectsToDisplay.Add(PropertySet);
+	Toolkit = MakeShared<FUVEditorModeToolkit>();
 }
 
-void UUVEditorMode::CreateToolkit()
+UObject* UUVEditorMode::GetBackgroundSettingsObject()
 {
-	Toolkit = MakeShared<FUVEditorModeToolkit>();
+	if (BackgroundVisualization)
+	{
+		return BackgroundVisualization->Settings;
+	}
+	return nullptr;
 }
 
 void UUVEditorMode::ActivateDefaultTool()
@@ -323,9 +323,7 @@ void UUVEditorMode::Exit()
 		BackgroundVisualization = nullptr;
 	}
 
-	UVChannelProperties = nullptr;
 	PropertyObjectsToTick.Empty();
-	PropertyObjectsToDisplay.Empty();
 	LivePreviewWorld = nullptr;
 
 	bIsActive = false;
@@ -459,20 +457,21 @@ void UUVEditorMode::InitializeTargets(const TArray<TObjectPtr<UObject>>& AssetsI
 		ToolInputObjects.Add(ToolInputObject);
 	}
 
-	// Initialize our layer selector
-	UVChannelProperties->Initialize(ToolTargets, AppliedCanonicalMeshes, true);
+	// Prep things for layer/channel selection
+	InitializeAssetNames(ToolTargets, AssetNames);
 	PendingUVLayerIndex.SetNumZeroed(ToolTargets.Num());		
 
-	int32 AssetID = UVChannelProperties->GetSelectedAssetID();
+	// Prep channel/layer change API
 	UUVToolAssetAndChannelAPI* AssetAndLayerAPI = ContextStore->FindContext<UUVToolAssetAndChannelAPI>();
 	if (AssetAndLayerAPI)
 	{
-		AssetAndLayerAPI->RequestChannelVisibilityChangeFunc = [this](const TArray<int32>& LayerPerAsset, bool bForceRebuildUnwrap, bool bEmitUndoTransaction) {
-			this->ForceUpdateDisplayChannel(LayerPerAsset, bForceRebuildUnwrap, bEmitUndoTransaction);
+		AssetAndLayerAPI->RequestChannelVisibilityChangeFunc = [this](const TArray<int32>& LayerPerAsset, bool bEmitUndoTransaction) {
+			this->SetDisplayedUVChannels(LayerPerAsset, bEmitUndoTransaction);
 		};
 
 		AssetAndLayerAPI->NotifyOfAssetChannelCountChangeFunc = [this](int32 AssetID) {
-			UVChannelProperties->Initialize(ToolTargets, AppliedCanonicalMeshes, false);
+			// Don't currently need to do anything because the layer selection menu gets populated
+			// from scratch each time that it's opened.
 		};
 
 		AssetAndLayerAPI->GetCurrentChannelVisibilityFunc = [this]() {
@@ -480,7 +479,7 @@ void UUVEditorMode::InitializeTargets(const TArray<TObjectPtr<UObject>>& AssetsI
 			VisibleLayers.SetNum(ToolTargets.Num());
 			for (int32 AssetID = 0; AssetID < ToolTargets.Num(); ++AssetID)
 			{
-				VisibleLayers[AssetID] = ToolInputObjects[AssetID]->UVLayerIndex;
+				VisibleLayers[AssetID] = GetDisplayedChannel(AssetID);
 			}
 			return VisibleLayers;
 		};
@@ -559,7 +558,7 @@ void UUVEditorMode::ModeTick(float DeltaTime)
 	bool bSwitchingLayers = false;
 	for (int32 AssetID = 0; AssetID < ToolInputObjects.Num(); ++AssetID)
 	{
-		if (ToolInputObjects[AssetID]->UVLayerIndex != PendingUVLayerIndex[AssetID] || bForceRebuildUVLayer)
+		if (ToolInputObjects[AssetID]->UVLayerIndex != PendingUVLayerIndex[AssetID])
 		{
 			bSwitchingLayers = true;
 			break;
@@ -568,31 +567,21 @@ void UUVEditorMode::ModeTick(float DeltaTime)
 
 	if (bSwitchingLayers)
 	{
-		GetToolManager()->BeginUndoTransaction(UVLayerChangeTransactionName);
-
 		// TODO: Perhaps we need our own interactive tools context that allows this kind of "end tool now"
 		// call. We can't do the normal GetInteractiveToolsContext()->EndTool() call because we cannot defer
 		// shutdown.
 		GetInteractiveToolsContext()->ToolManager->DeactivateTool(EToolSide::Mouse, EToolShutdownType::Cancel);
 
-		for (int32 AssetID = 0; AssetID < ToolInputObjects.Num(); ++AssetID)
-		{
-			if (ToolInputObjects[AssetID]->UVLayerIndex != PendingUVLayerIndex[AssetID] || bForceRebuildUVLayer)
-			{
-				int32 OldLayerIndex = ToolInputObjects[AssetID]->UVLayerIndex;
+		// We open the transaction here instead of before the tool close above because otherwise we seem
+		// to get some kind of spurious items in the transaction that, while they don't seem to do anything,
+		// cause the transaction to not be fully expired on editor close.
+		GetToolManager()->BeginUndoTransaction(UVLayerChangeTransactionName);
 
-				ChangeInputObjectLayer(AssetID, PendingUVLayerIndex[AssetID], bForceRebuildUVLayer);
+		SetDisplayedUVChannels(PendingUVLayerIndex, true);
 
-				GetInteractiveToolsContext()->GetTransactionAPI()->AppendChange(this,
-					MakeUnique<FInputObjectUVLayerChange>(AssetID, OldLayerIndex, PendingUVLayerIndex[AssetID]),
-					UVLayerChangeTransactionName);
-			}
-		}
 		ActivateDefaultTool();
 
 		GetToolManager()->EndUndoTransaction();
-
-		bForceRebuildUVLayer = false;
 	}
 	
 	for (TObjectPtr<UInteractiveToolPropertySet>& Propset : PropertyObjectsToTick)
@@ -631,102 +620,64 @@ void UUVEditorMode::ModeTick(float DeltaTime)
 	}
 }
 
-
-void UUVEditorMode::SwitchActiveAsset(const FString& UVAsset)
+int32 UUVEditorMode::GetNumUVChannels(int32 AssetID) const
 {
-	if (UVChannelProperties)
+	if (AssetID < 0 || AssetID >= AppliedCanonicalMeshes.Num())
 	{
-		// Not doing an ensure here because the "revert to default" can give us an empty string
-		UVChannelProperties->ValidateUVAssetSelection(true);
-		UpdateSelectedLayer();
-		bForceRebuildUVLayer = false;
+		return IndexConstants::InvalidID;
 	}
+
+	return AppliedCanonicalMeshes[AssetID]->HasAttributes() ?
+		AppliedCanonicalMeshes[AssetID]->Attributes()->NumUVLayers() : 0;
 }
 
-void UUVEditorMode::UpdateSelectedLayer()
+int32 UUVEditorMode::GetDisplayedChannel(int32 AssetID) const
 {
-	int32 AssetID = UVChannelProperties->GetSelectedAssetID();
-	if (!ensure(AssetID != INDEX_NONE))
+	if (!ensure(AssetID >= 0 && AssetID < ToolInputObjects.Num()))
 	{
-		return;
+		return IndexConstants::InvalidID;
 	}
-	const TArray<FString>& ChannelNames = UVChannelProperties->GetUVChannelNames();
 
-	// Find the currently selected layer for the current asset
-	bool bFound = false;
-	for (auto InputObject : ToolInputObjects)
-	{
-		if (InputObject->AssetID == AssetID)
-		{
-			bFound = true;
-
-			UVChannelProperties->UVChannel = ChannelNames[InputObject->UVLayerIndex];
-			UVChannelProperties->SilentUpdateWatched();
-			break;
-		}
-	}
-	if (!ensure(bFound))
-	{
-		UVChannelProperties->UVChannel = TEXT("");
-	}
-	else 
-	{
-		PendingUVLayerIndex[AssetID] = ToolInputObjects[AssetID]->UVLayerIndex;
-	}		
+	return ToolInputObjects[AssetID]->UVLayerIndex;;
 }
 
-
-void UUVEditorMode::SwitchActiveChannel(const FString& UVChannel)
+void UUVEditorMode::RequestUVChannelChange(int32 AssetID, int32 Channel)
 {
-	using namespace UVEditorModeLocals;
-
-	if (UVChannelProperties)
+	if (!ensure(
+		AssetID >= 0 && AssetID < ToolTargets.Num())
+		&& Channel >= 0 && Channel < GetNumUVChannels(AssetID))
 	{
-		// Not doing an ensure because the "revert to default" can give us an empty string
-		UVChannelProperties->ValidateUVChannelSelection(true);
-		int32 NewUVLayerIndex = UVChannelProperties->GetSelectedChannelIndex();
-		int32 AssetID = UVChannelProperties->GetSelectedAssetID();
-		if (!ensure(AssetID != INDEX_NONE))
-		{
-			return;
-		}
-
-		if (ensure(NewUVLayerIndex != IndexConstants::InvalidID))
-		{
-			PendingUVLayerIndex[AssetID] = NewUVLayerIndex;
-		}
+		return;
 	}
-	bForceRebuildUVLayer = false;
+
+	PendingUVLayerIndex[AssetID] = Channel;
 }
 
-void UUVEditorMode::ChangeInputObjectLayer(int32 AssetID, int32 NewLayerIndex, bool bForceRebuild)
+void UUVEditorMode::ChangeInputObjectLayer(int32 AssetID, int32 NewLayerIndex)
 {
 	using namespace UVEditorModeLocals;
 
-	bool bFound = false;
-	for (auto InputObject : ToolInputObjects)
+	if (!ensure(AssetID >= 0 && AssetID < ToolInputObjects.Num() 
+		&& ToolInputObjects[AssetID]->AssetID == AssetID))
 	{
-		if (InputObject->AssetID == AssetID)
-		{
-			bFound = true;
+		return;
+	}
 
-			if (InputObject->UVLayerIndex != NewLayerIndex || bForceRebuild)
-			{
-				InputObject->UVLayerIndex = NewLayerIndex;
-				InputObject->UpdateAllFromAppliedCanonical();
-			}
-			break;
-		}
+	UUVEditorToolMeshInput* InputObject = ToolInputObjects[AssetID];
+	if (InputObject->UVLayerIndex != NewLayerIndex)
+	{
+		InputObject->UVLayerIndex = NewLayerIndex;
+		InputObject->UpdateAllFromAppliedCanonical();
 	}
 
-	ensure(bFound);
+	PendingUVLayerIndex[AssetID] = InputObject->UVLayerIndex;
 }
 
-void UUVEditorMode::ForceUpdateDisplayChannel(const TArray<int32>& LayerPerAsset, bool bForceRebuildUnwrap, bool bEmitUndoTransaction)
+void UUVEditorMode::SetDisplayedUVChannels(const TArray<int32>& LayerPerAsset, bool bEmitUndoTransaction)
 {
 	for (int32 AssetID = 0; AssetID < ToolInputObjects.Num(); ++AssetID)
 	{
-		if (ToolInputObjects[AssetID]->UVLayerIndex != LayerPerAsset[AssetID] || bForceRebuildUnwrap)
+		if (ToolInputObjects[AssetID]->UVLayerIndex != LayerPerAsset[AssetID])
 		{			
 			if (bEmitUndoTransaction)
 			{
@@ -735,11 +686,9 @@ void UUVEditorMode::ForceUpdateDisplayChannel(const TArray<int32>& LayerPerAsset
 					UVEditorModeLocals::UVLayerChangeTransactionName);
 			}
 
-			ChangeInputObjectLayer(AssetID, LayerPerAsset[AssetID], true);
-			PendingUVLayerIndex[AssetID] = LayerPerAsset[AssetID];
+			ChangeInputObjectLayer(AssetID, LayerPerAsset[AssetID]);
 		}
 	}
-	UpdateSelectedLayer();
 }
 
 #undef LOCTEXT_NAMESPACE
diff --git a/Engine/Plugins/Experimental/UVEditor/Source/UVEditor/Private/UVEditorModeToolkit.cpp b/Engine/Plugins/Experimental/UVEditor/Source/UVEditor/Private/UVEditorModeToolkit.cpp
--- a/Engine/Plugins/Experimental/UVEditor/Source/UVEditor/Private/UVEditorModeToolkit.cpp
+++ b/Engine/Plugins/Experimental/UVEditor/Source/UVEditor/Private/UVEditorModeToolkit.cpp
@@ -8,6 +8,7 @@
 #include "Framework/MultiBox/MultiBoxDefs.h"
 #include "IDetailsView.h"
 #include "Modules/ModuleManager.h"
+#include "SPrimaryButton.h"
 #include "Tools/UEdMode.h"
 #include "UVEditorBackgroundPreview.h"
 #include "UVEditorCommands.h"
@@ -16,7 +17,6 @@
 #include "Widgets/Layout/SBorder.h"
 #include "Widgets/Text/STextBlock.h"
 #include "EdModeInteractiveToolsContext.h"
-#include "SPrimaryButton.h"
 
 #define LOCTEXT_NAMESPACE "FUVEditorModeToolkit"
 
@@ -75,16 +75,6 @@ FUVEditorModeToolkit::FUVEditorModeToolkit()
 		.Padding(4)
 		[
 			SNew(SVerticalBox)
-
-			+ SVerticalBox::Slot()
-			.AutoHeight()
-			[
-				SAssignNew(ToolMessageArea, STextBlock)
-				.AutoWrapText(true)
-			.Font(FCoreStyle::GetDefaultFontStyle("Bold", 9))
-			.Text(LOCTEXT("UVEditorToolsLabel", "UV Editor Tools"))
-			]
-
 			+ SVerticalBox::Slot()
 			.AutoHeight()
 			.HAlign(HAlign_Left)
@@ -116,47 +106,6 @@ FUVEditorModeToolkit::FUVEditorModeToolkit()
 					SAssignNew(ToolDetailsContainer, SBorder)
 					.BorderImage(FEditorStyle::GetBrush("NoBorder"))
 				]
-
-				+ SVerticalBox::Slot()
-				.AutoHeight()
-				[
-					SAssignNew(ToolMessageArea, STextBlock)
-					.AutoWrapText(true)
-					.Font(FCoreStyle::GetDefaultFontStyle("Bold", 9))
-					.Text(FText::GetEmpty())
-				]
-			]
-
-            // Todo: Move this out of here once we figure out how to add menus/UI to the viewport
-			+ SVerticalBox::Slot()
-				.AutoHeight()
-				.Padding(2, 0, 0, 0)
-			[
-					SNew(SVerticalBox)
-
-					+ SVerticalBox::Slot()
-					.AutoHeight()
-					[
-						SAssignNew(ToolMessageArea, STextBlock)
-						.AutoWrapText(true)
-						.Font(FCoreStyle::GetDefaultFontStyle("Bold", 9))
-						.Text(LOCTEXT("UVEditorSettingsLabel", "UV Editor Settings"))
-					]
-
-					+ SVerticalBox::Slot()
-					.AutoHeight()
-					[
-						SAssignNew(EditorDetailsContainer, SBorder)
-						.BorderImage(FEditorStyle::GetBrush("NoBorder"))
-					]
-
-					+ SVerticalBox::Slot()
-					.AutoHeight()
-					[
-						SAssignNew(BackgroundDetailsContainer, SBorder)
-						.BorderImage(FEditorStyle::GetBrush("NoBorder"))
-					]
-
 			]
 
 		];
@@ -172,11 +121,13 @@ void FUVEditorModeToolkit::Init(const TSharedPtr<IToolkitHost>& InitToolkitHost,
 
 	UUVEditorMode* UVEditorMode = Cast<UUVEditorMode>(GetScriptableEditorMode());
 	check(UVEditorMode);
+	UEditorInteractiveToolsContext* InteractiveToolsContext = UVEditorMode->GetInteractiveToolsContext();
+	check(InteractiveToolsContext);
 
 	// Currently, there's no EToolChangeTrackingMode that reverts back to a default tool on undo (if we add that
 	// support, the tool manager will need to be aware of the default tool). So, we instead opt to do our own management
 	// of tool start transactions. See FUVEditorModeToolkit::OnToolStarted for how we issue the transactions.
-	UVEditorMode->GetInteractiveToolsContext()->ToolManager->ConfigureChangeTrackingMode(EToolChangeTrackingMode::NoChangeTracking);
+	InteractiveToolsContext->ToolManager->ConfigureChangeTrackingMode(EToolChangeTrackingMode::NoChangeTracking);
 
 	// Build the tool palette
 	const FUVEditorCommands& CommandInfos = FUVEditorCommands::Get();
@@ -196,27 +147,6 @@ void FUVEditorModeToolkit::Init(const TSharedPtr<IToolkitHost>& InitToolkitHost,
 
 	// Hook up the tool detail panel
 	ToolDetailsContainer->SetContent(DetailsView.ToSharedRef());
-
-	// Hook up the editor detail panel
-	EditorDetailsContainer->SetContent(ModeDetailsView.ToSharedRef());
-
-	// Hook up the background detail panel
-	{
-		FPropertyEditorModule& PropertyEditorModule = FModuleManager::GetModuleChecked<FPropertyEditorModule>("PropertyEditor");
-
-		FDetailsViewArgs BackgroundDetailsViewArgs;
-		BackgroundDetailsViewArgs.bAllowSearch = false;
-		BackgroundDetailsViewArgs.NameAreaSettings = FDetailsViewArgs::HideNameArea;
-		BackgroundDetailsViewArgs.bHideSelectionTip = true;
-		BackgroundDetailsViewArgs.DefaultsOnlyVisibility = EEditDefaultsOnlyNodeVisibility::Automatic;
-		BackgroundDetailsViewArgs.bShowOptions = false;
-		BackgroundDetailsViewArgs.bAllowMultipleTopLevelObjects = true;
-
-		CustomizeDetailsViewArgs(BackgroundDetailsViewArgs);		// allow subclass to customize arguments
-
-		BackgroundDetailsView = PropertyEditorModule.CreateDetailView(BackgroundDetailsViewArgs);
-		BackgroundDetailsContainer->SetContent(BackgroundDetailsView.ToSharedRef());
-	}
 	
 	// Set up the overlay. Largely copied from ModelingToolsEditorModeToolkit.
 	// TODO: We could put some of the shared code in some common place.
@@ -323,10 +253,96 @@ FText FUVEditorModeToolkit::GetBaseToolkitName() const
 	return NSLOCTEXT("UVEditorModeToolkit", "DisplayName", "UVEditorMode");
 }
 
-void FUVEditorModeToolkit::SetModeDetailsViewObjects(const TArray<TObjectPtr<UInteractiveToolPropertySet>>& InObjects)
+TSharedRef<SWidget> FUVEditorModeToolkit::CreateChannelMenu()
 {
-	TArray<TObjectPtr<UObject>> Objects(InObjects);
-	ModeDetailsView->SetObjects(Objects);
+	UUVEditorMode* Mode = Cast<UUVEditorMode>(GetScriptableEditorMode());
+
+	const bool bShouldCloseWindowAfterMenuSelection = true;
+	FMenuBuilder MenuBuilder(bShouldCloseWindowAfterMenuSelection, nullptr);
+
+	// For each asset, create a submenu labeled with its name
+	const TArray<FString>& AssetNames = Mode->GetAssetNames();
+	for (int32 AssetID = 0; AssetID < AssetNames.Num(); ++AssetID)
+	{
+		MenuBuilder.AddSubMenu(
+			FText::AsCultureInvariant(AssetNames[AssetID]), // label
+			FText(), // tooltip
+			FNewMenuDelegate::CreateWeakLambda(Mode, [this, Mode, AssetID](FMenuBuilder& SubMenuBuilder)
+		{
+
+			// Inside each submenu, create a button for each channel
+			int32 NumChannels = Mode->GetNumUVChannels(AssetID);
+			for (int32 Channel = 0; Channel < NumChannels; ++Channel)
+			{
+				SubMenuBuilder.AddMenuEntry(
+					FText::Format(LOCTEXT("ChannelLabel", "UV Channel {0}"), Channel), // label
+					FText(), // tooltip
+					FSlateIcon(),
+					FUIAction(
+						FExecuteAction::CreateWeakLambda(Mode, [Mode, AssetID, Channel]()
+						{
+							Mode->RequestUVChannelChange(AssetID, Channel);
+
+							// A bit of a hack to force the menu to close if the checkbox is clicked (which usually doesn't
+							// close the menu)
+							FSlateApplication::Get().DismissAllMenus();
+						}), 
+						FCanExecuteAction::CreateWeakLambda(Mode, []() { return true; }),
+						FIsActionChecked::CreateWeakLambda(Mode, [Mode, AssetID, Channel]()
+						{
+							return Mode->GetDisplayedChannel(AssetID) == Channel;
+						})),
+					NAME_None,
+					EUserInterfaceActionType::RadioButton);
+			}
+		}));
+	}
+
+	return MenuBuilder.MakeWidget();
+}
+
+TSharedRef<SWidget> FUVEditorModeToolkit::CreateBackgroundSettingsWidget()
+{
+	TSharedRef<SBorder> BackgroundDetailsContainer =
+		SNew(SBorder)
+		.BorderImage(FEditorStyle::GetBrush("NoBorder"));
+
+	TSharedRef<SWidget> Widget = SNew(SBorder)
+		.HAlign(HAlign_Fill)
+		.Padding(4)
+		[
+			SNew(SVerticalBox)
+
+			+ SVerticalBox::Slot()
+				.AutoHeight()
+				.Padding(2, 0, 0, 0)
+			[
+				SNew(SVerticalBox)
+
+				+ SVerticalBox::Slot()
+				.AutoHeight()
+				[
+					BackgroundDetailsContainer
+				]
+			]
+		];
+
+	FPropertyEditorModule& PropertyEditorModule = FModuleManager::GetModuleChecked<FPropertyEditorModule>("PropertyEditor");
+
+	FDetailsViewArgs BackgroundDetailsViewArgs;
+	BackgroundDetailsViewArgs.bAllowSearch = false;
+	BackgroundDetailsViewArgs.NameAreaSettings = FDetailsViewArgs::HideNameArea;
+	BackgroundDetailsViewArgs.bHideSelectionTip = true;
+	BackgroundDetailsViewArgs.DefaultsOnlyVisibility = EEditDefaultsOnlyNodeVisibility::Automatic;
+	BackgroundDetailsViewArgs.bShowOptions = false;
+	BackgroundDetailsViewArgs.bAllowMultipleTopLevelObjects = false;
+
+	UUVEditorMode* Mode = Cast<UUVEditorMode>(GetScriptableEditorMode());
+	TSharedRef<IDetailsView> BackgroundDetailsView = PropertyEditorModule.CreateDetailView(BackgroundDetailsViewArgs);
+	BackgroundDetailsView->SetObject(Mode->GetBackgroundSettingsObject());
+	BackgroundDetailsContainer->SetContent(BackgroundDetailsView);
+
+	return Widget;
 }
 
 void FUVEditorModeToolkit::UpdateActiveToolProperties()
diff --git a/Engine/Plugins/Experimental/UVEditor/Source/UVEditor/Private/UVEditorToolkit.cpp b/Engine/Plugins/Experimental/UVEditor/Source/UVEditor/Private/UVEditorToolkit.cpp
--- a/Engine/Plugins/Experimental/UVEditor/Source/UVEditor/Private/UVEditorToolkit.cpp
+++ b/Engine/Plugins/Experimental/UVEditor/Source/UVEditor/Private/UVEditorToolkit.cpp
@@ -13,10 +13,12 @@
 #include "SUVEditor3DViewport.h"
 #include "ThumbnailRendering/SceneThumbnailInfo.h"
 #include "ToolContextInterfaces.h"
+#include "Toolkits/AssetEditorModeUILayer.h"
 #include "ToolMenus.h"
 #include "UVEditor.h"
-#include "UVEditorMode.h"
 #include "UVEditorCommands.h"
+#include "UVEditorMode.h"
+#include "UVEditorModeToolkit.h"
 #include "UVEditorSubsystem.h"
 #include "UVEditor2DViewportClient.h"
 #include "UVEditor3DViewportClient.h"
@@ -398,7 +400,7 @@ void FUVEditorToolkit::PostInitAssetEditor()
 	}
 
 	// Plug in the mode tool panel
-	TSharedPtr<FModeToolkit> UVModeToolkit = UVMode->GetToolkit().Pin();
+	TSharedPtr<FUVEditorModeToolkit> UVModeToolkit = StaticCastSharedPtr<FUVEditorModeToolkit>(UVMode->GetToolkit().Pin());
 	if (ensure(ToolsPanel && UVModeToolkit.IsValid()))
 	{
 		ToolsPanel->SetContent(UVModeToolkit->GetInlineContent().ToSharedRef());
@@ -421,8 +423,34 @@ void FUVEditorToolkit::PostInitAssetEditor()
 	FName ParentToolbarName;
 	const FName ToolBarName = GetToolMenuToolbarName(ParentToolbarName);
 	UToolMenu* AssetToolbar = UToolMenus::Get()->ExtendMenu(ToolBarName);
-	FToolMenuSection& Section = AssetToolbar->FindOrAddSection("UVAsset");
+	FToolMenuSection& Section = AssetToolbar->FindOrAddSection("Asset");
 	Section.AddEntry(FToolMenuEntry::InitToolBarButton(FUVEditorCommands::Get().ApplyChanges));
+	
+	// Add the channel selection button.
+	Section.AddEntry(FToolMenuEntry::InitComboButton(
+		"UVEditorChannelMenu",
+		FUIAction(),
+		FOnGetContent::CreateLambda([UVModeToolkit]()
+		{
+			return UVModeToolkit->CreateChannelMenu();
+		}),
+		LOCTEXT("UVEditorChannelMenu_Label", "Channels"),
+		LOCTEXT("UVEditorChannelMenu_ToolTip", "Change the selected UV Channel for each asset."),
+		FSlateIcon()
+	));
+
+	// Add the background settings button.
+	Section.AddEntry(FToolMenuEntry::InitComboButton(
+		"UVEditorBackgroundSettings",
+		FUIAction(),
+		FOnGetContent::CreateLambda([UVModeToolkit]()
+		{
+			return UVModeToolkit->CreateBackgroundSettingsWidget();
+		}),
+		LOCTEXT("UVEditorBackgroundSettings_Label", "Background"),
+		LOCTEXT("UVEditorBackgroundSettings_ToolTip", "Change the visibility and other settings of the background."),
+		FSlateIcon()
+	));
 
 
 	// Adjust our main (2D) viewport:
diff --git a/Engine/Plugins/Experimental/UVEditor/Source/UVEditor/Public/UVEditorBackgroundPreview.h b/Engine/Plugins/Experimental/UVEditor/Source/UVEditor/Public/UVEditorBackgroundPreview.h
new file mode 100644
--- /dev/null
+++ b/Engine/Plugins/Experimental/UVEditor/Source/UVEditor/Public/UVEditorBackgroundPreview.h
@@ -0,0 +1,75 @@
+// Copyright Epic Games, Inc. All Rights Reserved.
+
+#pragma once
+
+#include "CoreMinimal.h"
+#include "Drawing/PreviewGeometryActor.h"
+#include "Drawing/TriangleSetComponent.h"
+#include "InteractiveTool.h"
+#include "GeometryBase.h"
+#include "UVEditorBackgroundPreview.generated.h"
+
+/**
+ * Visualization settings for the UUVEditorBackgroundPreview
+ */
+UCLASS()
+class UVEDITOR_API UUVEditorBackgroundPreviewProperties : public UInteractiveToolPropertySet
+{
+	GENERATED_BODY()
+public:
+	/** Should the background be shown */
+	UPROPERTY(EditAnywhere, Category = Background)
+	bool bVisible = false;
+
+	/** Should the background show textures or materials */
+	UPROPERTY(EditAnywhere, Category = Background, meta = (EditCondition = "bVisible"))
+	bool bUseMaterials = false;
+
+	/** Used when bUseMaterials is false */
+	UPROPERTY(EditAnywhere, Category = Background, meta = (EditCondition = "!bUseMaterials && bVisible"))
+	TObjectPtr<UTexture2D> SourceTexture;
+
+	/** Used when bUseMaterials is true */
+	UPROPERTY(EditAnywhere, Category = Background, meta = (EditCondition = "bUseMaterials && bVisible"))
+	TObjectPtr<UMaterial> SourceMaterial;
+};
+
+/**
+  Serves as a container for the background texture/material display in the UVEditor. This class is responsible for managing the quad
+  drawn behind the grid, as well as maintaining the texture and material choices from the user to display.
+ */
+UCLASS(Transient)
+class UVEDITOR_API UUVEditorBackgroundPreview : public UPreviewGeometry
+{
+	GENERATED_BODY()
+
+public:
+
+	/**
+	 * Client must call this every frame for changes to .Settings to be reflected in rendered result.
+	 */
+	void OnTick(float DeltaTime);
+
+
+public:
+	/** Visualization settings */
+	UPROPERTY()
+	TObjectPtr<UUVEditorBackgroundPreviewProperties> Settings;
+
+	/** The component containing the quad visualization */
+	UPROPERTY()
+	TObjectPtr<UTriangleSetComponent> BackgroundComponent;
+
+	/** The active material being displayed for the background */
+	UPROPERTY()
+	TObjectPtr<UMaterialInstanceDynamic> BackgroundMaterial;
+
+protected:
+	virtual void OnCreated() override;
+
+	bool bSettingsModified = false;
+
+	void UpdateVisibility();
+	void UpdateBackground();
+};
+
diff --git a/Engine/Plugins/Experimental/UVEditor/Source/UVEditor/Public/UVEditorMode.h b/Engine/Plugins/Experimental/UVEditor/Source/UVEditor/Public/UVEditorMode.h
--- a/Engine/Plugins/Experimental/UVEditor/Source/UVEditor/Public/UVEditorMode.h
+++ b/Engine/Plugins/Experimental/UVEditor/Source/UVEditor/Public/UVEditorMode.h
@@ -24,7 +24,6 @@ class UMeshOpPreviewWithBackgroundCompute;
 class UUVToolStateObjectStore;
 class UUVEditorToolMeshInput;
 class UUVEditorBackgroundPreview;
-class UUVEditorUVChannelProperties;
 
 
 /**
@@ -56,9 +55,14 @@ public:
 
 	void InitializeTargets(const TArray<TObjectPtr<UObject>>& AssetsIn, const TArray<FTransform>& TransformsIn);
 
-	// public for use by undo/redo
-	void ChangeInputObjectLayer(int32 AssetID, int32 NewLayerIndex, bool bForceRebuild=false);
-	void UpdateSelectedLayer();
+	// Public for use by undo/redo. Otherwise should use RequestUVChannelChange
+	void ChangeInputObjectLayer(int32 AssetID, int32 NewLayerIndex);
+
+	/** 
+	 * Request a change of the displayed UV channel/layer. It will happen on the next tick, and 
+	 * create an undo/redo event.
+	 */
+	void RequestUVChannelChange(int32 AssetID, int32 Channel);
 
 	bool IsActive() { return bIsActive; }
 
@@ -78,6 +82,18 @@ public:
 	void GetAssetsWithUnappliedChanges(TArray<TObjectPtr<UObject>> UnappliedAssetsOut);
 	void ApplyChanges();
 
+	/** @return List of asset names, indexed by AssetID */
+	const TArray<FString>& GetAssetNames() const { return AssetNames; }
+
+	/** @return Number of UV channels in the given asset, or IndexConstants::InvalidID if AssetID was invalid.  */
+	int32 GetNumUVChannels(int32 AssetID) const;
+
+	/** @return The index of the channel currently displayed for the given AssetID. */
+	int32 GetDisplayedChannel(int32 AssetID) const;
+
+	/** @return A settings object suitable for display in a details panel to control the background visualization. */
+	UObject* GetBackgroundSettingsObject();
+
 	// UEdMode overrides
 	virtual void Enter() override;
 	virtual void Exit() override;
@@ -108,7 +124,6 @@ protected:
 	virtual void OnToolEnded(UInteractiveToolManager* Manager, UInteractiveTool* Tool) override {}
 	
 	void UpdateTriangleMaterialBasedOnBackground(bool IsBackgroundVisible);
-	void AddDisplayedPropertySet(const TObjectPtr<UInteractiveToolPropertySet>& PropertySet);
 
 	/**
 	 * Stores original input objects, for instance UStaticMesh pointers. AssetIDs on tool input 
@@ -164,24 +179,14 @@ protected:
 	// Authoritative list of targets that have changes that have not been baked back yet.
 	TSet<int32> ModifiedAssetIDs;
 
-
-	/** Used as a selector of UV channels/layers of opened assets in the editor. */
-	UPROPERTY()
-	TObjectPtr<UUVEditorUVChannelProperties> UVChannelProperties = nullptr;
-
-	// Used with UVChannelProperties
-	void SwitchActiveAsset(const FString& UVAsset);
-	void SwitchActiveChannel(const FString& UVChannel);
+	// 1:1 with ToolTargets, indexed by AssetID
+	TArray<FString> AssetNames;
 
 	// Used with the ToolAssetAndLayerAPI to process tool layer change requests
-	void ForceUpdateDisplayChannel(const TArray<int32>& LayerPerAsset, bool bForceRebuildUnwrap, bool bEmitUndoTransaction);
+	void SetDisplayedUVChannels(const TArray<int32>& LayerPerAsset, bool bEmitUndoTransaction);
 
 	TArray<int32> PendingUVLayerIndex;
 
-	// TODO: Should be removed
-	bool  bForceRebuildUVLayer = false;
-
-
 	// Wireframe Display Properties
 	float TriangleOpacity = 1.0;
 	float TriangleDepthOffset = 0.5;
@@ -194,12 +199,6 @@ protected:
 	UPROPERTY()
 	TObjectPtr<UWorld> LivePreviewWorld = nullptr;
 
-	/**
-	* Mode-level property objects to display in the details panel.
-	*/
-	UPROPERTY()
-	TArray<TObjectPtr<UInteractiveToolPropertySet>> PropertyObjectsToDisplay;
-
 	/**
 	 * Mode-level property objects (visible or not) that get ticked.
 	 */
diff --git a/Engine/Plugins/Experimental/UVEditor/Source/UVEditor/Public/UVEditorModeToolkit.h b/Engine/Plugins/Experimental/UVEditor/Source/UVEditor/Public/UVEditorModeToolkit.h
--- a/Engine/Plugins/Experimental/UVEditor/Source/UVEditor/Public/UVEditorModeToolkit.h
+++ b/Engine/Plugins/Experimental/UVEditor/Source/UVEditor/Public/UVEditorModeToolkit.h
@@ -4,8 +4,8 @@
 
 #include "CoreMinimal.h"
 
-#include "Toolkits/BaseToolkit.h"
 #include "InteractiveTool.h"
+#include "Toolkits/BaseToolkit.h"
 
 class SBorder;
 class STextBlock;
@@ -21,21 +21,27 @@ public:
 	FUVEditorModeToolkit();
 	~FUVEditorModeToolkit();
 
-	/** FModeToolkit interface */
+	/** Creates a menu where the displayed UV Channel can be changed for each asset */
+	virtual TSharedRef<SWidget> CreateChannelMenu();
+	
+	/** Creates a widget where the background visualization can be changed. */
+	virtual TSharedRef<SWidget> CreateBackgroundSettingsWidget();
+
+	// FModeToolkit
 	virtual void Init(const TSharedPtr<IToolkitHost>& InitToolkitHost, TWeakObjectPtr<UEdMode> InOwningMode) override;
+	virtual FText GetActiveToolDisplayName() const override { return ActiveToolName; }
 
-	/** IToolkit interface */
+	// IToolkit
 	virtual FName GetToolkitFName() const override;
 	virtual FText GetBaseToolkitName() const override;
 	virtual TSharedPtr<class SWidget> GetInlineContent() const override { return ToolkitWidget; }
 
-	virtual void SetModeDetailsViewObjects(const TArray<TObjectPtr<UInteractiveToolPropertySet>>& InObjects);
-
-	virtual FText GetActiveToolDisplayName() const override { return ActiveToolName; }
+protected:
+	// FModeToolkit
 	virtual void OnToolStarted(UInteractiveToolManager* Manager, UInteractiveTool* Tool) override;
 	virtual void OnToolEnded(UInteractiveToolManager* Manager, UInteractiveTool* Tool) override;
 
-protected:
+
 	// The mode's entire toolbox, which gets returned by GetInlineContent()
 	TSharedPtr<SWidget> ToolkitWidget;
 
@@ -52,13 +58,6 @@ protected:
 	// A container for the editor settings
 	TSharedPtr<SBorder> EditorDetailsContainer;
 
-	// A container for the background settings
-	TSharedPtr<SBorder> BackgroundDetailsContainer;
-	TSharedPtr<IDetailsView> BackgroundDetailsView;
-
-	// A place for tools to write out any instructions
-	TSharedPtr<STextBlock> ToolMessageArea;
-
 	/** Contains the widget container for the Accept/Cancel buttons for tools. */
 	TSharedPtr<SWidget> ViewportOverlayWidget;
 
diff --git a/Engine/Plugins/Experimental/UVEditor/Source/UVEditorTools/Private/UVEditorLayerEditTool.cpp b/Engine/Plugins/Experimental/UVEditor/Source/UVEditorTools/Private/UVEditorLayerEditTool.cpp
--- a/Engine/Plugins/Experimental/UVEditor/Source/UVEditorTools/Private/UVEditorLayerEditTool.cpp
+++ b/Engine/Plugins/Experimental/UVEditor/Source/UVEditorTools/Private/UVEditorLayerEditTool.cpp
@@ -21,11 +21,41 @@ using namespace UE::Geometry;
 
 namespace UUVEditorChannelEditLocals
 {
-	
 	const FText UVChannelAddTransactionName = LOCTEXT("UVChannelAddTransactionName", "Add UV Channel");
 	const FText UVChannelCloneTransactionName = LOCTEXT("UVChannelCloneTransactionName", "Clone UV Channel");
 	const FText UVChannelDeleteTransactionName = LOCTEXT("UVChannelDeleteTransactionName", "Delete UV Channel");
 
+	void DeleteChannel(UUVEditorToolMeshInput* Target, int32 DeletedUVChannelIndex, 
+		UUVToolAssetAndChannelAPI* AssetAndChannelAPI, bool bClearInstead)
+	{
+		FDynamicMeshUVEditor DynamicMeshUVEditor(Target->AppliedCanonical.Get(), DeletedUVChannelIndex, false);
+
+		int32 NewChannelIndex = DeletedUVChannelIndex;
+		if (bClearInstead)
+		{
+			DynamicMeshUVEditor.SetPerTriangleUVs(0.0, nullptr);
+		}
+		else
+		{
+			NewChannelIndex = DynamicMeshUVEditor.RemoveUVLayer();
+		}
+
+		if (NewChannelIndex != Target->UVLayerIndex)
+		{
+			// The change of displayed layer will perform the needed update for us
+			AssetAndChannelAPI->NotifyOfAssetChannelCountChange(Target->AssetID);
+			TArray<int32> ChannelPerAsset = AssetAndChannelAPI->GetCurrentChannelVisibility();
+			ChannelPerAsset[Target->AssetID] = NewChannelIndex;
+			AssetAndChannelAPI->RequestChannelVisibilityChange(ChannelPerAsset, false);
+		}
+		else
+		{
+			// We're showing the same layer index, but it's now the next layer over, actually,
+			// so we need to update it.
+			Target->UpdateAllFromAppliedCanonical();
+		}
+	}
+
 	class FInputObjectUVChannelAdd : public FToolCommandChange
 	{
 	public:
@@ -37,7 +67,11 @@ namespace UUVEditorChannelEditLocals
 
 		virtual void Apply(UObject* Object) override
 		{
-			UInteractiveToolManager* InteractiveToolManager = Cast<UInteractiveToolManager>(Object);
+			UInteractiveToolManager* InteractiveToolManager = Cast<UUVEditorChannelEditTool>(Object)->GetToolManager();
+			if (!ensure(InteractiveToolManager))
+			{
+				return;
+			}
 			UUVToolAssetAndChannelAPI* AssetAndChannelAPI = InteractiveToolManager->GetContextObjectStore()->FindContext<UUVToolAssetAndChannelAPI>();
 
 			FDynamicMeshUVEditor DynamicMeshUVEditor(Target->AppliedCanonical.Get(), AddedUVChannelIndex-1, false);
@@ -49,26 +83,30 @@ namespace UUVEditorChannelEditLocals
 				Target->AppliedPreview->PreviewMesh->UpdatePreview(Target->AppliedCanonical.Get());
 
 				AssetAndChannelAPI->NotifyOfAssetChannelCountChange(Target->AssetID);
-				Target->OnCanonicalModified.Broadcast(Target, UUVEditorToolMeshInput::FCanonicalModifiedInfo());
+				Target->OnCanonicalModified.Broadcast(Target.Get(), UUVEditorToolMeshInput::FCanonicalModifiedInfo());
 			}
 
 		}
 
 		virtual void Revert(UObject* Object) override
 		{
-			UInteractiveToolManager* InteractiveToolManager = Cast<UInteractiveToolManager>(Object);
+			UInteractiveToolManager* InteractiveToolManager = Cast<UUVEditorChannelEditTool>(Object)->GetToolManager();
+			if (!ensure(InteractiveToolManager))
+			{
+				return;
+			}
 			UUVToolAssetAndChannelAPI* AssetAndChannelAPI = InteractiveToolManager->GetContextObjectStore()->FindContext<UUVToolAssetAndChannelAPI>();
 
 			FDynamicMeshUVEditor DynamicMeshUVEditor(Target->AppliedCanonical.Get(), AddedUVChannelIndex, false);
 			int32 NewChannelIndex = DynamicMeshUVEditor.RemoveUVLayer();
 			Target->AppliedPreview->PreviewMesh->UpdatePreview(Target->AppliedCanonical.Get());
 			AssetAndChannelAPI->NotifyOfAssetChannelCountChange(Target->AssetID);
-			Target->OnCanonicalModified.Broadcast(Target, UUVEditorToolMeshInput::FCanonicalModifiedInfo());
+			Target->OnCanonicalModified.Broadcast(Target.Get(), UUVEditorToolMeshInput::FCanonicalModifiedInfo());
 		}
 
 		virtual bool HasExpired(UObject* Object) const override
 		{
-			return false;
+			return !(Target.IsValid() && Target->IsValid());
 		}
 
 		virtual FString ToString() const override
@@ -77,7 +115,7 @@ namespace UUVEditorChannelEditLocals
 		}
 
 	protected:
-		TObjectPtr<UUVEditorToolMeshInput> Target;
+		TWeakObjectPtr<UUVEditorToolMeshInput> Target;
 		int32 AddedUVChannelIndex;		
 	};
 
@@ -94,7 +132,11 @@ namespace UUVEditorChannelEditLocals
 
 		virtual void Apply(UObject* Object) override
 		{
-			UInteractiveToolManager* InteractiveToolManager = Cast<UInteractiveToolManager>(Object);
+			UInteractiveToolManager* InteractiveToolManager = Cast<UUVEditorChannelEditTool>(Object)->GetToolManager();
+			if (!ensure(InteractiveToolManager))
+			{
+				return;
+			}
 			UUVToolAssetAndChannelAPI* AssetAndChannelAPI = InteractiveToolManager->GetContextObjectStore()->FindContext<UUVToolAssetAndChannelAPI>();
 
 			FDynamicMeshUVEditor DynamicMeshUVEditor(Target->AppliedCanonical.Get(), TargetUVChannelIndex, false);
@@ -104,7 +146,11 @@ namespace UUVEditorChannelEditLocals
 
 		virtual void Revert(UObject* Object) override
 		{
-			UInteractiveToolManager* InteractiveToolManager = Cast<UInteractiveToolManager>(Object);
+			UInteractiveToolManager* InteractiveToolManager = Cast<UUVEditorChannelEditTool>(Object)->GetToolManager();
+			if (!ensure(InteractiveToolManager))
+			{
+				return;
+			}
 			UUVToolAssetAndChannelAPI* AssetAndChannelAPI = InteractiveToolManager->GetContextObjectStore()->FindContext<UUVToolAssetAndChannelAPI>();
 
 			FDynamicMeshUVEditor DynamicMeshUVEditor(Target->AppliedCanonical.Get(), TargetUVChannelIndex, false);
@@ -114,7 +160,7 @@ namespace UUVEditorChannelEditLocals
 
 		virtual bool HasExpired(UObject* Object) const override
 		{
-			return false;
+			return !(Target.IsValid() && Target->IsValid());
 		}
 
 		virtual FString ToString() const override
@@ -123,71 +169,81 @@ namespace UUVEditorChannelEditLocals
 		}
 
 	protected:
-		TObjectPtr<UUVEditorToolMeshInput> Target;
+		TWeakObjectPtr<UUVEditorToolMeshInput> Target;
 		int32 SourceUVChannelIndex;
 		int32 TargetUVChannelIndex;
 		FDynamicMeshUVOverlay OriginalUVChannel;
 	};
 
-
 	class FInputObjectUVChannelDelete : public FToolCommandChange
 	{
 	public:
-		FInputObjectUVChannelDelete(const TObjectPtr<UUVEditorToolMeshInput>& TargetIn, int32 DeletedUVChannelIndexIn, const FDynamicMeshUVOverlay& OriginalUVChannelIn)
+		FInputObjectUVChannelDelete(const TObjectPtr<UUVEditorToolMeshInput>& TargetIn, int32 DeletedUVChannelIndexIn, 
+			const FDynamicMeshUVOverlay& OriginalUVChannelIn, bool bClearedInsteadIn)
 			: Target(TargetIn)			
 			, DeletedUVChannelIndex(DeletedUVChannelIndexIn)
 			, OriginalUVChannel(OriginalUVChannelIn)
+			, bClearedInstead(bClearedInsteadIn)
 		{
 		}
 
 		virtual void Apply(UObject* Object) override
 		{
-			UInteractiveToolManager* InteractiveToolManager = Cast<UInteractiveToolManager>(Object);
+			UInteractiveToolManager* InteractiveToolManager = Cast<UUVEditorChannelEditTool>(Object)->GetToolManager();
+			if (!ensure(InteractiveToolManager))
+			{
+				return;
+			}
 			UUVToolAssetAndChannelAPI* AssetAndChannelAPI = InteractiveToolManager->GetContextObjectStore()->FindContext<UUVToolAssetAndChannelAPI>();
 
-			FDynamicMeshUVEditor DynamicMeshUVEditor(Target->AppliedCanonical.Get(), DeletedUVChannelIndex, false);
-
-			DynamicMeshUVEditor.SetPerTriangleUVs(0.0, nullptr); // "clear" the UVs to a single point.
-			int32 NewChannelIndex = DynamicMeshUVEditor.RemoveUVLayer();
-
-			Target->AppliedPreview->PreviewMesh->UpdatePreview(Target->AppliedCanonical.Get());
-
-			AssetAndChannelAPI->NotifyOfAssetChannelCountChange(Target->AssetID);
-			TArray<int32> ChannelPerAsset = AssetAndChannelAPI->GetCurrentChannelVisibility();
-			ChannelPerAsset[Target->AssetID] = DeletedUVChannelIndex;
-			AssetAndChannelAPI->RequestChannelVisibilityChange(ChannelPerAsset, false, false);
-			
-			Target->OnCanonicalModified.Broadcast(Target, UUVEditorToolMeshInput::FCanonicalModifiedInfo());
+			DeleteChannel(Target.Get(), DeletedUVChannelIndex, AssetAndChannelAPI, bClearedInstead);
 		}
 
 		virtual void Revert(UObject* Object) override
 		{
-			UInteractiveToolManager* InteractiveToolManager = Cast<UInteractiveToolManager>(Object);
+			UInteractiveToolManager* InteractiveToolManager = Cast<UUVEditorChannelEditTool>(Object)->GetToolManager();
+			if (!ensure(InteractiveToolManager))
+			{
+				return;
+			}
 			UUVToolAssetAndChannelAPI* AssetAndChannelAPI = InteractiveToolManager->GetContextObjectStore()->FindContext<UUVToolAssetAndChannelAPI>();
 
-			FDynamicMeshUVEditor DynamicMeshUVEditor(Target->AppliedCanonical.Get(), DeletedUVChannelIndex, false);
-			int32 NewChannelIndex = DynamicMeshUVEditor.AddUVLayer();
-			Target->AppliedPreview->PreviewMesh->UpdatePreview(Target->AppliedCanonical.Get());
-
-			for (int32 ChannelIndex = NewChannelIndex; ChannelIndex > DeletedUVChannelIndex; --ChannelIndex)
+			FDynamicMeshUVEditor DynamicMeshUVEditor(Target->AppliedCanonical.Get(), 0, false);
+			if (!bClearedInstead)
 			{
-				DynamicMeshUVEditor.SwitchActiveLayer(ChannelIndex);
-				DynamicMeshUVEditor.CopyUVLayer(Target->AppliedCanonical->Attributes()->GetUVLayer(ChannelIndex - 1));
+				int32 NewChannelIndex = DynamicMeshUVEditor.AddUVLayer();
+
+				// Shift the new layer to the correct index.
+				// TODO: This is slightly wasteful since we should just be allowed to swap pointers in the underlying indirect array
+				for (int32 ChannelIndex = NewChannelIndex; ChannelIndex > DeletedUVChannelIndex; --ChannelIndex)
+				{
+					DynamicMeshUVEditor.SwitchActiveLayer(ChannelIndex);
+					DynamicMeshUVEditor.CopyUVLayer(Target->AppliedCanonical->Attributes()->GetUVLayer(ChannelIndex - 1));
+				}
 			}
+
+			// Fill the layer
 			DynamicMeshUVEditor.SwitchActiveLayer(DeletedUVChannelIndex);
 			DynamicMeshUVEditor.CopyUVLayer(&OriginalUVChannel);
 
-			AssetAndChannelAPI->NotifyOfAssetChannelCountChange(Target->AssetID);
-			TArray<int32> ChannelPerAsset = AssetAndChannelAPI->GetCurrentChannelVisibility();
-			ChannelPerAsset[Target->AssetID] = DeletedUVChannelIndex;
-			AssetAndChannelAPI->RequestChannelVisibilityChange(ChannelPerAsset, true, false);
-			
-			Target->OnCanonicalModified.Broadcast(Target, UUVEditorToolMeshInput::FCanonicalModifiedInfo());
+			if (DeletedUVChannelIndex != Target->UVLayerIndex)
+			{
+				// The change of displayed layer will perform the needed update for us
+				AssetAndChannelAPI->NotifyOfAssetChannelCountChange(Target->AssetID);
+				TArray<int32> ChannelPerAsset = AssetAndChannelAPI->GetCurrentChannelVisibility();
+				ChannelPerAsset[Target->AssetID] = DeletedUVChannelIndex;
+				AssetAndChannelAPI->RequestChannelVisibilityChange(ChannelPerAsset, false);
+			}
+			else
+			{
+				// We're showing the same layer index, but it's actually the previously deleted layer
+				Target->UpdateAllFromAppliedCanonical();
+			}
 		}
 
 		virtual bool HasExpired(UObject* Object) const override
 		{
-			return false;
+			return !(Target.IsValid() && Target->IsValid());
 		}
 
 		virtual FString ToString() const override
@@ -196,9 +252,10 @@ namespace UUVEditorChannelEditLocals
 		}
 
 	protected:
-		TObjectPtr<UUVEditorToolMeshInput> Target;		
+		TWeakObjectPtr<UUVEditorToolMeshInput> Target;		
 		int32 DeletedUVChannelIndex;
 		FDynamicMeshUVOverlay OriginalUVChannel;
+		bool bClearedInstead;
 	};
 
 }
@@ -537,7 +594,7 @@ void UUVEditorChannelEditTool::ApplyVisbleChannelChange()
 	{
 		TArray<int32> ChannelPerAsset = AssetAndChannelAPI->GetCurrentChannelVisibility();
 		ChannelPerAsset[ActiveAsset] = ActiveChannel;
-		AssetAndChannelAPI->RequestChannelVisibilityChange(ChannelPerAsset);
+		AssetAndChannelAPI->RequestChannelVisibilityChange(ChannelPerAsset, true);
 	}
 }
 
@@ -571,14 +628,14 @@ void UUVEditorChannelEditTool::AddChannel()
 		Targets[ActiveAsset]->AppliedPreview->PreviewMesh->UpdatePreview(Targets[ActiveAsset]->AppliedCanonical.Get());
 
 		EmitChangeAPI->BeginUndoTransaction(UUVEditorChannelEditLocals::UVChannelAddTransactionName);
-		EmitChangeAPI->EmitToolIndependentChange(GetToolManager(),
+		EmitChangeAPI->EmitToolIndependentChange(this,
 			MakeUnique<UUVEditorChannelEditLocals::FInputObjectUVChannelAdd>(Targets[ActiveAsset], NewChannelIndex),
 			UUVEditorChannelEditLocals::UVChannelAddTransactionName);
 
 		AssetAndChannelAPI->NotifyOfAssetChannelCountChange(ActiveAsset);
 		TArray<int32> ChannelPerAsset = AssetAndChannelAPI->GetCurrentChannelVisibility();
 		ChannelPerAsset[ActiveAsset] = NewChannelIndex;
-		AssetAndChannelAPI->RequestChannelVisibilityChange(ChannelPerAsset);
+		AssetAndChannelAPI->RequestChannelVisibilityChange(ChannelPerAsset, true);
 
 		EmitChangeAPI->EndUndoTransaction();
 
@@ -592,22 +649,14 @@ void UUVEditorChannelEditTool::CopyChannel()
 {
 	UUVToolAssetAndChannelAPI* AssetAndChannelAPI = GetToolManager()->GetContextObjectStore()->FindContext<UUVToolAssetAndChannelAPI>();
 
-	FDynamicMeshUVOverlay OriginalUVChannel(Targets[ActiveAsset]->AppliedCanonical.Get());
-	OriginalUVChannel.Copy(*Targets[ActiveAsset]->AppliedCanonical->Attributes()->GetUVLayer(ActiveChannel));
-
-	FDynamicMeshUVEditor DynamicMeshUVEditor(Targets[ActiveAsset]->AppliedCanonical.Get(), ActiveChannel, false);
-	DynamicMeshUVEditor.CopyUVLayer( Targets[ActiveAsset]->AppliedCanonical->Attributes()->GetUVLayer(ReferenceChannel) );
-
-	EmitChangeAPI->BeginUndoTransaction(UUVEditorChannelEditLocals::UVChannelCloneTransactionName);
-	EmitChangeAPI->EmitToolIndependentChange(GetToolManager(),
-		MakeUnique<UUVEditorChannelEditLocals::FInputObjectUVChannelClone>(Targets[ActiveAsset], ReferenceChannel, ActiveChannel, OriginalUVChannel),
+	EmitChangeAPI->EmitToolIndependentChange(this,
+		MakeUnique<UUVEditorChannelEditLocals::FInputObjectUVChannelClone>(Targets[ActiveAsset], ReferenceChannel, ActiveChannel, 
+			*Targets[ActiveAsset]->AppliedCanonical->Attributes()->GetUVLayer(ActiveChannel)),
 		UUVEditorChannelEditLocals::UVChannelCloneTransactionName);
 
-	TArray<int32> ChannelPerAsset = AssetAndChannelAPI->GetCurrentChannelVisibility();
-	ChannelPerAsset[ActiveAsset] = ActiveChannel;
-	AssetAndChannelAPI->RequestChannelVisibilityChange(ChannelPerAsset, true);
-
-	EmitChangeAPI->EndUndoTransaction();
+	FDynamicMeshUVEditor DynamicMeshUVEditor(Targets[ActiveAsset]->AppliedCanonical.Get(), ActiveChannel, false);
+	DynamicMeshUVEditor.CopyUVLayer(Targets[ActiveAsset]->AppliedCanonical->Attributes()->GetUVLayer(ReferenceChannel));
+	Targets[ActiveAsset]->UpdateAllFromAppliedCanonical();
 
 	SourceChannelProperties->Initialize(Targets, false);
 }
@@ -615,31 +664,19 @@ void UUVEditorChannelEditTool::CopyChannel()
 void UUVEditorChannelEditTool::DeleteChannel()
 {
 	UUVToolAssetAndChannelAPI* AssetAndChannelAPI = GetToolManager()->GetContextObjectStore()->FindContext<UUVToolAssetAndChannelAPI>();
-	FDynamicMeshUVEditor DynamicMeshUVEditor(Targets[ActiveAsset]->AppliedCanonical.Get(), ActiveChannel, false);
-	
-	FDynamicMeshUVOverlay OriginalUVChannel(Targets[ActiveAsset]->AppliedCanonical.Get());
-	OriginalUVChannel.Copy(*Targets[ActiveAsset]->AppliedCanonical->Attributes()->GetUVLayer(ActiveChannel));
-
-	DynamicMeshUVEditor.SetPerTriangleUVs(0.0, nullptr); // "clear" the UVs to a single point.
-	int32 NewChannelIndex = DynamicMeshUVEditor.RemoveUVLayer();
-	
-	Targets[ActiveAsset]->AppliedPreview->PreviewMesh->UpdatePreview(Targets[ActiveAsset]->AppliedCanonical.Get());
-
-	EmitChangeAPI->BeginUndoTransaction(UUVEditorChannelEditLocals::UVChannelDeleteTransactionName);
-	EmitChangeAPI->EmitToolIndependentChange(GetToolManager(),
-		MakeUnique<UUVEditorChannelEditLocals::FInputObjectUVChannelDelete>(Targets[ActiveAsset], ActiveChannel, OriginalUVChannel),
-		UUVEditorChannelEditLocals::UVChannelDeleteTransactionName);
 
-	AssetAndChannelAPI->NotifyOfAssetChannelCountChange(ActiveAsset);
-	TArray<int32> ChannelPerAsset = AssetAndChannelAPI->GetCurrentChannelVisibility();
-	ChannelPerAsset[ActiveAsset] = NewChannelIndex;
-	AssetAndChannelAPI->RequestChannelVisibilityChange(ChannelPerAsset, true);
+	int32 TotalLayerCount = Targets[ActiveAsset]->AppliedCanonical->Attributes()->NumUVLayers();
+	bool bClearInstead = TotalLayerCount == 1;
 
-	EmitChangeAPI->EndUndoTransaction();
+	EmitChangeAPI->EmitToolIndependentChange(this,
+		MakeUnique<UUVEditorChannelEditLocals::FInputObjectUVChannelDelete>(Targets[ActiveAsset], ActiveChannel, 
+			*Targets[ActiveAsset]->AppliedCanonical->Attributes()->GetUVLayer(ActiveChannel), bClearInstead),
+		UUVEditorChannelEditLocals::UVChannelDeleteTransactionName);
 
+	UUVEditorChannelEditLocals::DeleteChannel(Targets[ActiveAsset], ActiveChannel, AssetAndChannelAPI, bClearInstead);
 
 	SourceChannelProperties->Initialize(Targets, false);
-	SourceChannelProperties->TargetChannel = SourceChannelProperties->GetUVChannelNames()[NewChannelIndex];
+	SourceChannelProperties->TargetChannel = SourceChannelProperties->GetUVChannelNames()[Targets[ActiveAsset]->UVLayerIndex];
 }
 
 #undef LOCTEXT_NAMESPACE
\ No newline at end of file
diff --git a/Engine/Plugins/Experimental/UVEditor/Source/UVEditorTools/Public/UVToolContextObjects.h b/Engine/Plugins/Experimental/UVEditor/Source/UVEditorTools/Public/UVToolContextObjects.h
--- a/Engine/Plugins/Experimental/UVEditor/Source/UVEditorTools/Public/UVToolContextObjects.h
+++ b/Engine/Plugins/Experimental/UVEditor/Source/UVEditorTools/Public/UVToolContextObjects.h
@@ -180,11 +180,11 @@ public:
 		return TArray<int32>();
 	}
 
-	void RequestChannelVisibilityChange(const TArray<int32>& ChannelPerAsset, bool bForceRebuildUnwrap=false, bool bEmitUndoTransaction=true)
+	void RequestChannelVisibilityChange(const TArray<int32>& ChannelPerAsset, bool bEmitUndoTransaction=true)
 	{
 		if (RequestChannelVisibilityChangeFunc)
 		{
-			RequestChannelVisibilityChangeFunc(ChannelPerAsset, bForceRebuildUnwrap, bEmitUndoTransaction);
+			RequestChannelVisibilityChangeFunc(ChannelPerAsset, bEmitUndoTransaction);
 		}
 	}
 
@@ -198,7 +198,7 @@ public:
 
 
 	TUniqueFunction<TArray<int32>()> GetCurrentChannelVisibilityFunc;
-	TUniqueFunction<void(const TArray<int32>&, bool, bool)> RequestChannelVisibilityChangeFunc;
+	TUniqueFunction<void(const TArray<int32>&, bool)> RequestChannelVisibilityChangeFunc;
 	TUniqueFunction<void(int32 AssetID)> NotifyOfAssetChannelCountChangeFunc;
 
 };
