UVEditor: Adjust colors of triangles and wireframe in the unwrap view to support per asset colors, procedurally assigned upon startup. Colors are also more muted, with direct (in code) controls over hue, saturation, and value.

#rnx

diff --git a/Engine/Plugins/Experimental/UVEditor/Source/UVEditor/Private/UVEditorMode.cpp b/Engine/Plugins/Experimental/UVEditor/Source/UVEditor/Private/UVEditorMode.cpp
--- a/Engine/Plugins/Experimental/UVEditor/Source/UVEditor/Private/UVEditorMode.cpp
+++ b/Engine/Plugins/Experimental/UVEditor/Source/UVEditor/Private/UVEditorMode.cpp
@@ -140,7 +140,7 @@ void UUVEditorMode::Enter()
 		[this](bool IsVisible) { UpdateTriangleMaterialBasedOnBackground(IsVisible); });
 
 	PropertyObjectsToTick.Add(BackgroundVisualization->Settings);
-	
+
 	RegisterTools();
 
 	ActivateDefaultTool();
@@ -415,7 +415,7 @@ void UUVEditorMode::InitializeTargets(const TArray<TObjectPtr<UObject>>& AssetsI
 		ToolInputObject->UnwrapPreview->PreviewMesh->SetMaterial(
 			0, ToolSetupUtil::GetCustomTwoSidedDepthOffsetMaterial(
 				GetToolManager(),
-				(FLinearColor)TriangleColor,
+				(FLinearColor)GetTriangleColorByTargetIndex(AssetID),
 				TriangleDepthOffset,
 				TriangleOpacity));
 
@@ -427,8 +427,8 @@ void UUVEditorMode::InitializeTargets(const TArray<TObjectPtr<UObject>>& AssetsI
 		WireframeDisplay->Settings->bAdjustDepthBiasUsingMeshSize = false;
 		WireframeDisplay->Settings->bShowWireframe = true;
 		WireframeDisplay->Settings->bShowBorders = true;
-		WireframeDisplay->Settings->WireframeColor = WireframeColor;
-		WireframeDisplay->Settings->BoundaryEdgeColor = IslandBorderColor;
+		WireframeDisplay->Settings->WireframeColor = GetWireframeColorByTargetIndex(AssetID).ToFColorSRGB();
+		WireframeDisplay->Settings->BoundaryEdgeColor = GetBoundaryColorByTargetIndex(AssetID).ToFColorSRGB(); ;
 		WireframeDisplay->Settings->bShowUVSeams = false;
 		WireframeDisplay->Settings->bShowNormalSeams = false;
 		// These are not exposed at the visualizer level yet
@@ -486,6 +486,37 @@ void UUVEditorMode::InitializeTargets(const TArray<TObjectPtr<UObject>>& AssetsI
 	}
 }
 
+FLinearColor UUVEditorMode::GetTriangleColorByTargetIndex(int32 TargetIndex) const
+{
+	double GoldenAngle = 137.50776405;
+
+	FLinearColor BaseColorHSV;
+	BaseColorHSV.R = FMath::Fmod(UnwrapBaseColorHueStart + (GoldenAngle / 2.0 * TargetIndex), 360);;
+	BaseColorHSV.G = UnwrapBaseColorSaturation;
+	BaseColorHSV.B = UnwrapBaseColorValue;
+	return BaseColorHSV.HSVToLinearRGB();
+}
+
+FLinearColor UUVEditorMode::GetWireframeColorByTargetIndex(int32 TargetIndex) const
+{
+	FLinearColor BaseColorHSV = GetTriangleColorByTargetIndex(TargetIndex).LinearRGBToHSV();
+	FLinearColor WireframeColorHSV = BaseColorHSV;
+	WireframeColorHSV.R = FMath::Fmod((WireframeColorHSV.R + UnwrapWireframeHueShift), 360);
+	WireframeColorHSV.G = UnwrapWireframeSaturation;
+	WireframeColorHSV.B = UnwrapWireframeValue;
+	return WireframeColorHSV.HSVToLinearRGB();
+}
+
+FLinearColor UUVEditorMode::GetBoundaryColorByTargetIndex(int32 TargetIndex) const
+{
+	FLinearColor BaseColorHSV = GetTriangleColorByTargetIndex(TargetIndex).LinearRGBToHSV();
+	FLinearColor BoundaryColorHSV = BaseColorHSV;
+	BoundaryColorHSV.R = FMath::Fmod( (BoundaryColorHSV.R + UnwrapBoundaryHueShift), 360);
+	BoundaryColorHSV.G = UnwrapBoundarySaturation;
+	BoundaryColorHSV.B = UnwrapBoundaryValue;
+	return BoundaryColorHSV.HSVToLinearRGB();
+}
+
 void UUVEditorMode::EmitToolIndependentObjectChange(UObject* TargetObject, TUniquePtr<FToolCommandChange> Change, const FText& Description)
 {
 	GetInteractiveToolsContext()->GetTransactionAPI()->AppendChange(TargetObject, MoveTemp(Change), Description);
@@ -538,12 +569,11 @@ void UUVEditorMode::UpdateTriangleMaterialBasedOnBackground(bool IsBackgroundVis
 
 	// Modify the material of the unwrapped mesh to account for the presence/absence of the background, 
 	// changing the opacity as set just above.
-	for (UUVEditorToolMeshInput* ToolInputObject : ToolInputObjects)
-	{
-		ToolInputObject->UnwrapPreview->PreviewMesh->SetMaterial(
+	for (int32 AssetID = 0; AssetID < ToolInputObjects.Num(); ++AssetID) {
+		ToolInputObjects[AssetID]->UnwrapPreview->PreviewMesh->SetMaterial(
 			0, ToolSetupUtil::GetCustomTwoSidedDepthOffsetMaterial(
 				GetToolManager(),
-				(FLinearColor)TriangleColor,
+				(FLinearColor)GetTriangleColorByTargetIndex(AssetID),
 				TriangleDepthOffset, 
 				TriangleOpacity));
 	}
diff --git a/Engine/Plugins/Experimental/UVEditor/Source/UVEditor/Public/UVEditorMode.h b/Engine/Plugins/Experimental/UVEditor/Source/UVEditor/Public/UVEditorMode.h
--- a/Engine/Plugins/Experimental/UVEditor/Source/UVEditor/Public/UVEditorMode.h
+++ b/Engine/Plugins/Experimental/UVEditor/Source/UVEditor/Public/UVEditorMode.h
@@ -191,9 +191,19 @@ protected:
 	float TriangleOpacity = 1.0;
 	float TriangleDepthOffset = 0.5;
 	const float WireframeDepthOffset = 0.6;
-	FColor TriangleColor = FColor(50, 194, 219);
-	FColor WireframeColor = FColor(50, 100, 219);
-	FColor IslandBorderColor = FColor(103, 52, 235);
+	float UnwrapBaseColorHueStart = 200;
+	float UnwrapBaseColorSaturation = 0.25;
+	float UnwrapBaseColorValue = 0.75;
+	float UnwrapWireframeHueShift = 15;
+	float UnwrapWireframeSaturation = 0.40;
+	float UnwrapWireframeValue = 0.40;
+	float UnwrapBoundaryHueShift = 30;
+	float UnwrapBoundarySaturation = 0.50;
+	float UnwrapBoundaryValue = 0.50;
+
+	FLinearColor GetTriangleColorByTargetIndex(int32 TargetIndex) const;
+	FLinearColor GetWireframeColorByTargetIndex(int32 TargetIndex) const;
+	FLinearColor GetBoundaryColorByTargetIndex(int32 TargetIndex) const;
 
 	// Here largely for convenience to avoid having to pass it around functions.
 	UPROPERTY()
