UVEditor: Add button icons.

#rnx

diff --git a/Engine/Plugins/Experimental/UVEditor/Source/UVEditor/Private/SUVEditor3DViewportToolBar.cpp b/Engine/Plugins/Experimental/UVEditor/Source/UVEditor/Private/SUVEditor3DViewportToolBar.cpp
--- a/Engine/Plugins/Experimental/UVEditor/Source/UVEditor/Private/SUVEditor3DViewportToolBar.cpp
+++ b/Engine/Plugins/Experimental/UVEditor/Source/UVEditor/Private/SUVEditor3DViewportToolBar.cpp
@@ -6,9 +6,9 @@
 
 #include "UVEditorCommands.h"
 #include "Framework/MultiBox/MultiBoxBuilder.h"
+#include "UVEditorStyle.h"
 #include "Widgets/Layout/SBorder.h"
 #include "Widgets/SBoxPanel.h"
-#include "EditorStyleSet.h"
 
 void SUVEditor3DViewportToolBar::Construct(const FArguments& InArgs)
 {
@@ -62,12 +62,12 @@ TSharedRef<SWidget> SUVEditor3DViewportToolBar::MakeToolBar(const TSharedPtr<FEx
 		// Orbit Camera
 		static FName OrbitCameraName = FName(TEXT("OrbitCamera"));
 		ToolbarBuilder.AddToolBarButton(FUVEditorCommands::Get().EnableOrbitCamera, NAME_None, TAttribute<FText>(), TAttribute<FText>(), 
-			TAttribute<FSlateIcon>(FSlateIcon(FEditorStyle::GetStyleSetName(), FName("EditorViewport.RotateMode"))), OrbitCameraName);
+			TAttribute<FSlateIcon>(FSlateIcon(FUVEditorStyle::Get().GetStyleSetName(), "UVEditor.OrbitCamera")), OrbitCameraName);
 
 		// Fly Camera
 		static FName FlyCameraName = FName(TEXT("FlyCamera"));
-		ToolbarBuilder.AddToolBarButton(FUVEditorCommands::Get().EnableFlyCamera, NAME_None, TAttribute<FText>(), TAttribute<FText>(), 
-			TAttribute<FSlateIcon>(FSlateIcon(FEditorStyle::GetStyleSetName(), FName("EditorViewport.CamSpeedSetting"))), FlyCameraName);
+		ToolbarBuilder.AddToolBarButton(FUVEditorCommands::Get().EnableFlyCamera, NAME_None, TAttribute<FText>(), TAttribute<FText>(),
+			TAttribute<FSlateIcon>(FSlateIcon(FUVEditorStyle::Get().GetStyleSetName(), "UVEditor.FlyCamera")), FlyCameraName);
 
 		ToolbarBuilder.EndBlockGroup();
 	}
diff --git a/Engine/Plugins/Experimental/UVEditor/Source/UVEditor/Private/UVEditorCommands.cpp b/Engine/Plugins/Experimental/UVEditor/Source/UVEditor/Private/UVEditorCommands.cpp
--- a/Engine/Plugins/Experimental/UVEditor/Source/UVEditor/Private/UVEditorCommands.cpp
+++ b/Engine/Plugins/Experimental/UVEditor/Source/UVEditor/Private/UVEditorCommands.cpp
@@ -4,14 +4,15 @@
 
 #include "EditorStyleSet.h"
 #include "Framework/Commands/InputChord.h"
+#include "UVEditorStyle.h"
 
 #define LOCTEXT_NAMESPACE "FUVEditorCommands"
 	
 FUVEditorCommands::FUVEditorCommands()
-	: TCommands<FUVEditorCommands>("UVEditorCommands", 
-		NSLOCTEXT("Contexts", "UVEditorCommands", "UV Editor"), 
+	: TCommands<FUVEditorCommands>("UVEditor",
+		LOCTEXT("ContextDescription", "UV Editor"), 
 		NAME_None, // Parent
-		FEditorStyle::GetStyleSetName() // TODO: What should go here?
+		FUVEditorStyle::Get().GetStyleSetName()
 		)
 {
 }
diff --git a/Engine/Plugins/Experimental/UVEditor/Source/UVEditor/Private/UVEditorModule.cpp b/Engine/Plugins/Experimental/UVEditor/Source/UVEditor/Private/UVEditorModule.cpp
--- a/Engine/Plugins/Experimental/UVEditor/Source/UVEditor/Private/UVEditorModule.cpp
+++ b/Engine/Plugins/Experimental/UVEditor/Source/UVEditor/Private/UVEditorModule.cpp
@@ -8,6 +8,7 @@
 #include "UVEditor.h"
 #include "UVEditorCommands.h"
 #include "UVEditorMode.h"
+#include "UVEditorStyle.h"
 #include "UVEditorSubsystem.h"
 #include "EditorModeRegistry.h"
 
@@ -17,6 +18,7 @@
 
 void FUVEditorModule::StartupModule()
 {
+	FUVEditorStyle::Get(); // Causes the constructor to be called
 	FUVEditorCommands::Register();
 
 	// Menus need to be registered in a callback to make sure the system is ready for them.
diff --git a/Engine/Plugins/Experimental/UVEditor/Source/UVEditor/Private/UVEditorStyle.cpp b/Engine/Plugins/Experimental/UVEditor/Source/UVEditor/Private/UVEditorStyle.cpp
new file mode 100644
--- /dev/null
+++ b/Engine/Plugins/Experimental/UVEditor/Source/UVEditor/Private/UVEditorStyle.cpp
@@ -0,0 +1,52 @@
+// Copyright Epic Games, Inc. All Rights Reserved.
+
+#include "UVEditorStyle.h"
+
+#include "Styling/SlateTypes.h"
+#include "Styling/CoreStyle.h"
+#include "EditorStyleSet.h"
+#include "Styling/SlateStyleMacros.h"
+
+FName FUVEditorStyle::StyleName("UVStyle");
+
+FUVEditorStyle::FUVEditorStyle()
+	: FSlateStyleSet(StyleName)
+{
+	// Used FFractureEditorStyle as a model
+
+	const FVector2D IconSize(20.0f, 20.0f);
+
+	SetContentRoot(FPaths::EnginePluginsDir() / TEXT("Experimental/UVEditor/Content/Icons"));
+	SetCoreContentRoot(FPaths::EngineContentDir() / TEXT("Editor/Slate") );
+
+	Set("UVEditor.BeginSelectTool",             new IMAGE_BRUSH_SVG("MeshSelect", IconSize));
+	Set("UVEditor.BeginLayoutTool",             new IMAGE_BRUSH_SVG("UVLayout", IconSize));
+	Set("UVEditor.BeginParameterizeMeshTool",   new IMAGE_BRUSH_SVG("AutoUnwrap", IconSize));
+	Set("UVEditor.BeginChannelEditTool",        new IMAGE_BRUSH_SVG("AttributeEditor", IconSize));
+	Set("UVEditor.BeginSeamTool",               new IMAGE_BRUSH_SVG("ModelingUVSeamEdit", IconSize));
+	Set("UVEditor.BeginRecomputeUVsTool",       new IMAGE_BRUSH_SVG("GroupUnwrap", IconSize));
+
+	// Top toolbar icons
+	Set("UVEditor.ApplyChanges", new CORE_IMAGE_BRUSH_SVG("Starship/Common/Apply", IconSize));
+	Set("UVEditor.ChannelSettings", new CORE_IMAGE_BRUSH_SVG("Starship/Common/SetDrawUVs", IconSize));
+	Set("UVEditor.BackgroundSettings", new CORE_IMAGE_BRUSH_SVG("Starship/Common/Sprite", IconSize));
+
+	// Viewport icons
+	Set("UVEditor.OrbitCamera", new CORE_IMAGE_BRUSH_SVG("Starship/EditorViewport/rotate", IconSize));
+	Set("UVEditor.FlyCamera", new CORE_IMAGE_BRUSH_SVG("Starship/EditorViewport/camera", IconSize));
+
+	FSlateStyleRegistry::RegisterSlateStyle(*this);
+}
+
+FUVEditorStyle::~FUVEditorStyle()
+{
+	FSlateStyleRegistry::UnRegisterSlateStyle(*this);
+}
+
+FUVEditorStyle& FUVEditorStyle::Get()
+{
+	static FUVEditorStyle Inst;
+	return Inst;
+}
+
+
diff --git a/Engine/Plugins/Experimental/UVEditor/Source/UVEditor/Private/UVEditorToolkit.cpp b/Engine/Plugins/Experimental/UVEditor/Source/UVEditor/Private/UVEditorToolkit.cpp
--- a/Engine/Plugins/Experimental/UVEditor/Source/UVEditor/Private/UVEditorToolkit.cpp
+++ b/Engine/Plugins/Experimental/UVEditor/Source/UVEditor/Private/UVEditorToolkit.cpp
@@ -16,12 +16,13 @@
 #include "Toolkits/AssetEditorModeUILayer.h"
 #include "ToolMenus.h"
 #include "UVEditor.h"
+#include "UVEditor2DViewportClient.h"
+#include "UVEditor3DViewportClient.h"
 #include "UVEditorCommands.h"
 #include "UVEditorMode.h"
 #include "UVEditorModeToolkit.h"
 #include "UVEditorSubsystem.h"
-#include "UVEditor2DViewportClient.h"
-#include "UVEditor3DViewportClient.h"
+#include "UVEditorStyle.h"
 #include "UVToolContextObjects.h"
 #include "Widgets/Docking/SDockTab.h"
 
@@ -436,7 +437,7 @@ void FUVEditorToolkit::PostInitAssetEditor()
 		}),
 		LOCTEXT("UVEditorChannelMenu_Label", "Channels"),
 		LOCTEXT("UVEditorChannelMenu_ToolTip", "Change the selected UV Channel for each asset."),
-		FSlateIcon()
+		FSlateIcon(FUVEditorStyle::Get().GetStyleSetName(), "UVEditor.ChannelSettings")
 	));
 
 	// Add the background settings button.
@@ -449,7 +450,7 @@ void FUVEditorToolkit::PostInitAssetEditor()
 		}),
 		LOCTEXT("UVEditorBackgroundSettings_Label", "Background"),
 		LOCTEXT("UVEditorBackgroundSettings_ToolTip", "Change the visibility and other settings of the background."),
-		FSlateIcon()
+		FSlateIcon(FUVEditorStyle::Get().GetStyleSetName(), "UVEditor.BackgroundSettings")
 	));
 
 
diff --git a/Engine/Plugins/Experimental/UVEditor/Source/UVEditor/Public/UVEditorStyle.h b/Engine/Plugins/Experimental/UVEditor/Source/UVEditor/Public/UVEditorStyle.h
new file mode 100644
--- /dev/null
+++ b/Engine/Plugins/Experimental/UVEditor/Source/UVEditor/Public/UVEditorStyle.h
@@ -0,0 +1,24 @@
+// Copyright Epic Games, Inc. All Rights Reserved.
+
+#pragma once
+
+#include "Styling/SlateStyle.h"
+#include "Styling/SlateStyleRegistry.h"
+
+/**
+ * Slate style set for UV Editor
+ */
+class UVEDITOR_API FUVEditorStyle
+	: public FSlateStyleSet
+{
+public:
+	static FName StyleName;
+
+	/** Access the singleton instance for this style set */
+	static FUVEditorStyle& Get();
+
+private:
+
+	FUVEditorStyle();
+	~FUVEditorStyle();
+};
\ No newline at end of file
