UVEditor: Fix selected edges being incorrectly set after a split action. Fix sew action passing deleted verts to update function.

#rnx

diff --git a/Engine/Plugins/Experimental/UVEditor/Source/UVEditorTools/Private/UVSeamSewAction.cpp b/Engine/Plugins/Experimental/UVEditor/Source/UVEditorTools/Private/UVSeamSewAction.cpp
--- a/Engine/Plugins/Experimental/UVEditor/Source/UVEditorTools/Private/UVSeamSewAction.cpp
+++ b/Engine/Plugins/Experimental/UVEditor/Source/UVEditorTools/Private/UVSeamSewAction.cpp
@@ -260,13 +260,21 @@ bool UUVSeamSewAction::ApplyAction(UUVToolEmitChangeAPI& EmitChangeAPI)
 
 	TArray<FDynamicMesh3::FMergeEdgesInfo> AllMergeInfo;
 	
+	// Note that currently, we don't really need to gather up the kept verts and call them out for an
+	// update, since no vert locations should have changed. But we'll do it anyway in case any of the
+	// code changes in some way that moves the remaining verts (for instance, to some halfway point).
+	TSet<int32> KeptVidsSet;
+
 	for (FIndex2i EdgePair : ResolvedEdgePairs)
 	{
 		FDynamicMesh3::FMergeEdgesInfo MergeInfo;
 		EMeshResult Result = MeshToSew.MergeEdges(EdgePair[0], EdgePair[1], MergeInfo, false);
 		if (Result == EMeshResult::Ok)
 		{
-			AllMergeInfo.Add(MergeInfo);
+			for (int i = 0; i < 2; ++i)
+			{
+				KeptVidsSet.Add(MergeInfo.KeptVerts[i]);
+			}
 		}
 		else
 		{
@@ -275,13 +283,15 @@ bool UUVSeamSewAction::ApplyAction(UUVToolEmitChangeAPI& EmitChangeAPI)
 	}
 	checkSlow(MeshToSew.CheckValidity(FDynamicMesh3::FValidityOptions(true, true))); // Allow nonmanifold verts and reverse orientation
 
-
 	TArray<int32> RemainingVids;
-	RemainingVids.SetNum(AllMergeInfo.Num() * 2);
-	int32 RemainingVidsIndex = 0;
-	for (FDynamicMesh3::FMergeEdgesInfo& MergeInfo : AllMergeInfo) {
-		RemainingVids[RemainingVidsIndex++] = MergeInfo.KeptVerts[0];
-		RemainingVids[RemainingVidsIndex++] = MergeInfo.KeptVerts[1];
+	for (int32 KeptVid : KeptVidsSet)
+	{
+		// We have to do this check because we may have performed multiple merge actions, so a "kept vert"
+		// from one action may have ended up getting removed in a later one.
+		if (MeshToSew.IsVertex(KeptVid))
+		{
+			RemainingVids.Add(KeptVid);
+		}
 	}
 
 	Targets[SelectionTargetIndex]->UpdateUnwrapCanonicalOverlayFromPositions(&RemainingVids, &SelectedTids);
diff --git a/Engine/Plugins/Experimental/UVEditor/Source/UVEditorTools/Private/UVSelectTool.cpp b/Engine/Plugins/Experimental/UVEditor/Source/UVEditorTools/Private/UVSelectTool.cpp
--- a/Engine/Plugins/Experimental/UVEditor/Source/UVEditorTools/Private/UVSelectTool.cpp
+++ b/Engine/Plugins/Experimental/UVEditor/Source/UVEditorTools/Private/UVSelectTool.cpp
@@ -777,11 +777,6 @@ void UUVSelectTool::GizmoTransformEnded(UTransformProxy* Proxy)
 			UUVEditorToolMeshInput::NONE_CHANGED_ARG, &SelectedTids);
 	}
 
-	if (!AABBTrees[SelectionTargetIndex]->IsValid(false))
-	{
-		AABBTrees[SelectionTargetIndex]->Build();
-	}
-
 	const FText TransactionName(LOCTEXT("DragCompleteTransactionName", "Move Items"));
 	EmitChangeAPI->EmitToolIndependentChange(ChangeRouter, MakeUnique<UVSelectToolLocals::FGizmoMeshChange>(
 		Targets[SelectionTargetIndex], ChangeTracker.EndChange(), 
@@ -920,6 +915,15 @@ void UUVSelectTool::ApplySplit()
 	TSet<int32> AppliedEidSet;
 	for (int32 Eid : Selection.SelectedIDs)
 	{
+		if (Selection.Mesh->IsBoundaryEdge(Eid))
+		{
+			// We will skip these already-split edges here. It would be safe to pass the corresponding
+			// applied edge into the CreateSeamsAtEdges call, but we don't want this edge to stay selected
+			// after the split action, because we would like a split followed by immediate sew to revert
+			// the mesh to the previous state, rather than sewing edges that started out split.
+			continue;
+		}
+
 		FIndex2i EdgeUnwrapVids = Selection.Mesh->GetEdgeV(Eid);
 
 		int32 AppliedEid = Target->AppliedCanonical->FindEdge(
@@ -951,17 +955,52 @@ void UUVSelectTool::ApplySplit()
 	ChangeTracker.BeginChange();
 	ChangeTracker.SaveTriangles(TidSet, true);
 
+	// We're about to update the unwrap, which may mess up our selection because a selected edge may
+	// no longer exist after the update, even if we store it as a pair of verts. Instead, we're going
+	// to be changing the selection to all the resulting border edges after this.
+	// The cleanest thing to do (esp for undo/redo) is to clear selection first, then reset it.
+
+	const FText TransactionName(LOCTEXT("ApplySplitTransactionName", "Split Edges"));
+	EmitChangeAPI->BeginUndoTransaction(TransactionName);
+
+	FDynamicMeshSelection NewSelection = SelectionMechanic->GetCurrentSelection();
+	FDynamicMeshSelection EmptySelection;
+	SelectionMechanic->SetSelection(EmptySelection, false, true); // don't broadcast, do emit undo
+
 	// Perform the update
 	TArray<int32> AppliedTids = TidSet.Array();
 	Target->UpdateAllFromAppliedCanonical(&UVEditResult.NewUVElements, &AppliedTids, &AppliedTids);
-	AABBTrees[SelectionTargetIndex]->Build();
+	
+	// Not needed because it should happen automatically via broadcast of target canonical mesh change
+	// AABBTrees[SelectionTargetIndex]->Build();
 
-	// Emit transaction
-	const FText TransactionName(LOCTEXT("ApplySplitTransactionName", "Split Edges"));
+	// Emit update transaction
 	EmitChangeAPI->EmitToolIndependentUnwrapCanonicalChange(
 		Target, ChangeTracker.EndChange(), TransactionName);
 
-	UpdateSelectionEidsAfterMeshChange(*SelectionMechanic, &CurrentSelectionVidPairs);
+	// Set selection to new border edges
+	NewSelection.SelectedIDs.Empty();
+	for (int32 AppliedEid : AppliedEidSet)
+	{
+		FIndex2i EdgeAppliedVids = Target->AppliedCanonical->GetEdgeV(AppliedEid);
+		TArray<int32> UnwrapVids1, UnwrapVids2;
+		Target->AppliedVidToUnwrapVids(EdgeAppliedVids.A, UnwrapVids1);
+		Target->AppliedVidToUnwrapVids(EdgeAppliedVids.B, UnwrapVids2);
+		for (int32 Vid1 : UnwrapVids1)
+		{
+			for (int32 Vid2 : UnwrapVids2)
+			{
+				int32 Eid = NewSelection.Mesh->FindEdge(Vid1, Vid2);
+				if (Eid != IndexConstants::InvalidID)
+				{
+					NewSelection.SelectedIDs.Add(Eid);
+				}
+			}
+		}
+	}
+	SelectionMechanic->SetSelection(NewSelection, true, true); // both broadcast and emit undo
+
+	EmitChangeAPI->EndUndoTransaction();
 }
 
 #undef LOCTEXT_NAMESPACE
