UVEditor: Move Unwrap tool into the UVEditor. Changes from the modeling mode version:

* Removed material override settings, since the editor can provide this functionality independently of tools.
* Removed Layout preview settings and support, since the editor visualizes the unwrap mesh already.

Known Issues:

* Editor currently doesn't handle unset UVs, causing crashes when using poly groups on meshes without poly groups
* Polygroup settings currently assume only one target. If multiple assets are loaded in the editor, polygroup settings will be removed and the default layer assumed for all loaded assets.

#rnx

diff --git a/Engine/Plugins/Experimental/MeshModelingToolsetExp/Source/MeshModelingToolsExp/Private/RecomputeUVsTool.cpp b/Engine/Plugins/Experimental/MeshModelingToolsetExp/Source/MeshModelingToolsExp/Private/RecomputeUVsTool.cpp
--- a/Engine/Plugins/Experimental/MeshModelingToolsetExp/Source/MeshModelingToolsExp/Private/RecomputeUVsTool.cpp
+++ b/Engine/Plugins/Experimental/MeshModelingToolsetExp/Source/MeshModelingToolsExp/Private/RecomputeUVsTool.cpp
@@ -42,8 +42,14 @@ void URecomputeUVsTool::Setup()
 	FComponentMaterialSet MaterialSet = UE::ToolTarget::GetMaterialSet(Target);
 	FTransform TargetTransform = (FTransform)UE::ToolTarget::GetLocalToWorldTransform(Target);
 
+	RecomputeUVsOpFactory = NewObject<URecomputeUVsOpFactory>();
+	RecomputeUVsOpFactory->OriginalMesh = InputMesh;
+	RecomputeUVsOpFactory->Settings = Settings;
+	RecomputeUVsOpFactory->TargetTransform = UE::ToolTarget::GetLocalToWorldTransform(Target);
+	RecomputeUVsOpFactory->GetSelectedUVChannel = [this]() { return GetSelectedUVChannel(); };
+
 	Preview = NewObject<UMeshOpPreviewWithBackgroundCompute>(this);
-	Preview->Setup(this->TargetWorld, this);
+	Preview->Setup(this->TargetWorld, RecomputeUVsOpFactory);
 	ToolSetupUtil::ApplyRenderingConfigurationToPreview(Preview->PreviewMesh, Target);
 	Preview->PreviewMesh->SetTangentsMode(EDynamicMeshComponentTangentsMode::AutoCalculated);
 	Preview->PreviewMesh->ReplaceMesh(*InputMesh);
@@ -154,6 +160,11 @@ void URecomputeUVsTool::Shutdown(EToolShutdownType ShutdownType)
 		GetToolManager()->EndUndoTransaction();
 	}
 
+	if (RecomputeUVsOpFactory)
+	{
+		RecomputeUVsOpFactory = nullptr;
+	}
+
 }
 
 void URecomputeUVsTool::Render(IToolsContextRenderAPI* RenderAPI)
@@ -180,60 +191,12 @@ bool URecomputeUVsTool::CanAccept() const
 	return Super::CanAccept() && Preview->HaveValidResult();
 }
 
-TUniquePtr<FDynamicMeshOperator> URecomputeUVsTool::MakeNewOperator()
-{
-	FAxisAlignedBox3d MeshBounds = Preview->PreviewMesh->GetMesh()->GetBounds();
-	TUniquePtr<FRecomputeUVsOp> RecomputeUVsOp = MakeUnique<FRecomputeUVsOp>();
-	RecomputeUVsOp->InputMesh = InputMesh;
-	RecomputeUVsOp->InputGroups = ActiveGroupSet;
-	RecomputeUVsOp->UVLayer = UVChannelProperties->GetSelectedChannelIndex(true);
-	
-	RecomputeUVsOp->IslandMode = (UE::Geometry::ERecomputeUVsIslandMode)(int)Settings->IslandMode;
-	RecomputeUVsOp->UnwrapType = (UE::Geometry::ERecomputeUVsUnwrapType)(int)Settings->UnwrapType;
-
-	RecomputeUVsOp->bAutoRotate = (Settings->AutoRotation == ERecomputeUVsToolOrientationMode::MinBoxBounds);
-
-	RecomputeUVsOp->NormalSmoothingRounds = Settings->SmoothingSteps;
-	RecomputeUVsOp->NormalSmoothingAlpha = Settings->SmoothingAlpha;
-
-	RecomputeUVsOp->bMergingOptimization = Settings->bIslandMerging;
-	RecomputeUVsOp->MergingThreshold = Settings->MergingThreshold;
-	RecomputeUVsOp->MaxNormalDeviationDeg = Settings->MaxAngleDeviation;
-
-	RecomputeUVsOp->bPackUVs = Settings->bAutoPack;
-	if (Settings->bAutoPack)
-	{
-		RecomputeUVsOp->PackingTextureResolution = Settings->TextureResolution;
-		RecomputeUVsOp->bNormalizeAreas = false;
-		RecomputeUVsOp->AreaScaling = 1.0;
-	}
-	else
-	{
-		switch (Settings->UVScaleMode)
-		{
-		case ERecomputeUVsToolUVScaleMode::NoScaling:
-			RecomputeUVsOp->bNormalizeAreas = false;
-			RecomputeUVsOp->AreaScaling = 1.0;
-			break;
-		case ERecomputeUVsToolUVScaleMode::NormalizeToBounds:
-			RecomputeUVsOp->bNormalizeAreas = true;
-			RecomputeUVsOp->AreaScaling = Settings->UVScale / MeshBounds.MaxDim();
-			break;
-		case ERecomputeUVsToolUVScaleMode::NormalizeToWorld:
-			RecomputeUVsOp->bNormalizeAreas = true;
-			RecomputeUVsOp->AreaScaling = Settings->UVScale;
-			break;
-		}
-	}
-
 
-	RecomputeUVsOp->SetTransform(UE::ToolTarget::GetLocalToWorldTransform(Target));
-
-	return RecomputeUVsOp;
+int32 URecomputeUVsTool::GetSelectedUVChannel() const
+{
+	return UVChannelProperties ? UVChannelProperties->GetSelectedChannelIndex(true) : 0;
 }
 
-
-
 void URecomputeUVsTool::OnSelectedGroupLayerChanged()
 {
 	UpdateActiveGroupLayer();
diff --git a/Engine/Plugins/Experimental/MeshModelingToolsetExp/Source/MeshModelingToolsExp/Public/RecomputeUVsTool.h b/Engine/Plugins/Experimental/MeshModelingToolsetExp/Source/MeshModelingToolsExp/Public/RecomputeUVsTool.h
--- a/Engine/Plugins/Experimental/MeshModelingToolsetExp/Source/MeshModelingToolsExp/Public/RecomputeUVsTool.h
+++ b/Engine/Plugins/Experimental/MeshModelingToolsetExp/Source/MeshModelingToolsExp/Public/RecomputeUVsTool.h
@@ -8,6 +8,7 @@
 #include "MeshOpPreviewHelpers.h"
 #include "Properties/MeshMaterialProperties.h"
 #include "Properties/MeshUVChannelProperties.h"
+#include "Properties/RecomputeUVsProperties.h"
 #include "PropertySets/PolygroupLayersProperties.h"
 #include "Polygroups/PolygroupSet.h"
 #include "Drawing/UVLayoutPreview.h"
@@ -19,6 +20,8 @@
 class UDynamicMeshComponent;
 class UMaterialInterface;
 class UMaterialInstanceDynamic;
+class RecomputeUVsOp;
+class URecomputeUVsOpFactory;
 
 
 /**
@@ -34,118 +37,11 @@ public:
 };
 
 
-
-UENUM()
-enum class ERecomputeUVsUnwrapType
-{
-	// Values must match UE::Geometry::ERecomputeUVsUnwrapType
-
-	/** ExpMap UV Flattening Algorithm is very fast but has limited ability to control stretching/distortion */
-	ExpMap = 0,
-	/** Conformal UV Flattening algorithm is increasingly-expensive on large islands but better at distortion control */
-	Conformal = 1
-};
-
-
-
-UENUM()
-enum class ERecomputeUVsIslandMode
-{
-	// Values must match UE::Geometry::ERecomputeUVsIslandMode
-
-	/** Use Active Mesh Polygroups Layer to define initial UV Islands */
-	PolyGroups = 0,
-	/** Use Existing UV Layer to define UV Islands (ie re-solve UV flattening based on existing UVs) */
-	ExistingUVs = 1
-};
-
-
-UENUM()
-enum class ERecomputeUVsToolOrientationMode
-{
-	/**  */
-	None,
-	/**  */
-	MinBoxBounds
-};
-
-
-UENUM()
-enum class ERecomputeUVsToolUVScaleMode
-{
-	/** No scaling is applied to UV islands */
-	NoScaling,
-	/** Scale UV islands such that they have constant relative area, relative to object bounds */
-	NormalizeToBounds,
-	/** Scale UV islands such that they have constant relative area, relative to world space */
-	NormalizeToWorld
-};
-
-
-UCLASS()
-class MESHMODELINGTOOLSEXP_API URecomputeUVsToolProperties : public UInteractiveToolPropertySet
-{
-	GENERATED_BODY()
-public:
-
-	/** IslandMode determines which patches of triangles will be used as initial UV Islands */
-	UPROPERTY(EditAnywhere, Category = "UV Generation")
-	ERecomputeUVsIslandMode IslandMode = ERecomputeUVsIslandMode::PolyGroups;
-
-	/** Which type of UV flattening algorithm to use */
-	UPROPERTY(EditAnywhere, Category = "UV Generation", meta = (EditCondition = "bIslandMerging == false"))
-	ERecomputeUVsUnwrapType UnwrapType = ERecomputeUVsUnwrapType::Conformal;
-
-	/** Which type of automatic island orientation to use */
-	UPROPERTY(EditAnywhere, Category = "UV Generation")
-	ERecomputeUVsToolOrientationMode AutoRotation = ERecomputeUVsToolOrientationMode::MinBoxBounds;
-
-	/** If Island Merging is enabled, then the initial UV islands will be merged into larger islands if it does not increase distortion/stretching beyond the limits below */
-	UPROPERTY(EditAnywhere, Category = "UV Generation")
-	bool bIslandMerging = false;
-
-	/** Distortion/Stretching Threshold for island merging - larger values increase the allowable UV stretching */
-	UPROPERTY(EditAnywhere, Category = "UV Generation", meta = (UIMin = "1.0", UIMax = "5.0", ClampMin = "1.0", EditCondition = "bIslandMerging == true"))
-	float MergingThreshold = 1.5f;
-
-	/** UV islands will not be merged if their average face normals deviate by larger than this amount */
-	UPROPERTY(EditAnywhere, Category = "UV Generation", meta = (UIMin = "0.0", UIMax = "90.0", ClampMin = "0.0", ClampMax = "180.0", EditCondition = "bIslandMerging == true"))
-	float MaxAngleDeviation = 45.0f;
-
-	/** Number of smoothing steps to apply in the ExpMap UV Generation method (Smoothing slightly increases distortion but produces more stable results) */
-	UPROPERTY(EditAnywhere, Category = "ExpMap UV Generator", meta = (UIMin = "0", UIMax = "25", ClampMin = "0", ClampMax = "1000", EditCondition = "UnwrapType == ERecomputeUVsUnwrapType::ExpMap || bIslandMerging == true" ))
-	int SmoothingSteps = 5;
-
-	/** Smoothing parameter, larger values result in faster smoothing in each step */
-	UPROPERTY(EditAnywhere, Category = "ExpMap UV Generator", meta = (UIMin = "0", UIMax = "1.0", ClampMin = "0", ClampMax = "1.0", EditCondition = "UnwrapType == ERecomputeUVsUnwrapType::ExpMap || bIslandMerging == true"))
-	float SmoothingAlpha = 0.25f;
-
-	/** If enabled, result UVs are automatically packed into the standard UV 0-1 square */
-	UPROPERTY(EditAnywhere, Category = "UV Layout")
-	bool bAutoPack = true;
-
-	/** Target texture resolution used for UV packing, which determines gutter size */
-	UPROPERTY(EditAnywhere, Category = "UV Layout", meta = (UIMin = "64", UIMax = "2048", ClampMin = "2", ClampMax = "4096", EditCondition = bAutoPack))
-	int TextureResolution = 1024;
-
-	/** Scaling applied to UV islands after solve */
-	UPROPERTY(EditAnywhere, Category = "Transform UVs", meta = (EditCondition = "bAutoPack == false"))
-	ERecomputeUVsToolUVScaleMode UVScaleMode = ERecomputeUVsToolUVScaleMode::NormalizeToBounds;
-
-	/** Scaling factor used for UV island normalization/scaling */
-	UPROPERTY(EditAnywhere, Category = "Transform UVs", meta = (EditCondition = "UVScaleMode!=ERecomputeUVsToolUVScaleMode::NoScaling && bAutoPack == false", UIMin = "0.001", UIMax = "10", ClampMin = "0.00001", ClampMax = "1000000.0") )
-	float UVScale = 1.0;
-};
-
-
-
-
-
 /**
  * URecomputeUVsTool Recomputes UVs based on existing segmentations of the mesh
  */
 UCLASS()
-class MESHMODELINGTOOLSEXP_API URecomputeUVsTool : public USingleSelectionMeshEditingTool, public UE::Geometry::IDynamicMeshOperatorFactory
+class MESHMODELINGTOOLSEXP_API URecomputeUVsTool : public USingleSelectionMeshEditingTool
 {
 	GENERATED_BODY()
 
@@ -162,9 +58,6 @@ public:
 
 	virtual void OnPropertyModified(UObject* PropertySet, FProperty* Property) override;
 
-	// IDynamicMeshOperatorFactory API
-	virtual TUniquePtr<UE::Geometry::FDynamicMeshOperator> MakeNewOperator() override;
-
 protected:
 	UPROPERTY()
 	TObjectPtr<UMeshUVChannelProperties> UVChannelProperties = nullptr;
@@ -184,6 +77,8 @@ protected:
 	UPROPERTY()
 	TObjectPtr<UUVLayoutPreview> UVLayoutView = nullptr;
 
+	UPROPERTY()
+	TObjectPtr<URecomputeUVsOpFactory> RecomputeUVsOpFactory;
 
 protected:
 	UPROPERTY()
@@ -196,6 +91,6 @@ protected:
 	TSharedPtr<UE::Geometry::FPolygroupSet, ESPMode::ThreadSafe> ActiveGroupSet;
 	void OnSelectedGroupLayerChanged();
 	void UpdateActiveGroupLayer();
-
+	int32 GetSelectedUVChannel() const;
 	void OnPreviewMeshUpdated();
 };
diff --git a/Engine/Plugins/Experimental/UVEditor/Source/UVEditor/Private/UVEditorCommands.cpp b/Engine/Plugins/Experimental/UVEditor/Source/UVEditor/Private/UVEditorCommands.cpp
--- a/Engine/Plugins/Experimental/UVEditor/Source/UVEditor/Private/UVEditorCommands.cpp
+++ b/Engine/Plugins/Experimental/UVEditor/Source/UVEditor/Private/UVEditorCommands.cpp
@@ -29,6 +29,7 @@ void FUVEditorCommands::RegisterCommands()
 	UI_COMMAND(BeginParameterizeMeshTool, "Auto UV", "Switch to Auto UV tool", EUserInterfaceActionType::ToggleButton, FInputChord());
 	UI_COMMAND(BeginChannelEditTool, "Channel Edit", "Switch to Channel Edit tool", EUserInterfaceActionType::ToggleButton, FInputChord());
 	UI_COMMAND(BeginSeamTool, "Seam", "Switch to Seam tool", EUserInterfaceActionType::ToggleButton, FInputChord());
+	UI_COMMAND(BeginRecomputeUVsTool, "Unwrap", "Switch to Unwrap tool", EUserInterfaceActionType::ToggleButton, FInputChord());
 
 	// These allow us to link up to pressed keys
 	UI_COMMAND(AcceptOrCompleteActiveTool, "Accept", "Accept the active tool", EUserInterfaceActionType::Button, FInputChord(EKeys::Enter));
diff --git a/Engine/Plugins/Experimental/UVEditor/Source/UVEditor/Private/UVEditorMode.cpp b/Engine/Plugins/Experimental/UVEditor/Source/UVEditor/Private/UVEditorMode.cpp
--- a/Engine/Plugins/Experimental/UVEditor/Source/UVEditor/Private/UVEditorMode.cpp
+++ b/Engine/Plugins/Experimental/UVEditor/Source/UVEditor/Private/UVEditorMode.cpp
@@ -26,6 +26,7 @@
 #include "UVEditorParameterizeMeshTool.h"
 #include "UVEditorLayerEditTool.h"
 #include "UVEditorSeamTool.h"
+#include "UVEditorRecomputeUVsTool.h"
 #include "UVSelectTool.h"
 #include "UVEditorModeToolkit.h"
 #include "UVEditorSubsystem.h"
@@ -181,6 +182,10 @@ void UUVEditorMode::RegisterTools()
 	UUVEditorSeamToolBuilder* UVEditorSeamToolBuilder = NewObject<UUVEditorSeamToolBuilder>();
 	UVEditorSeamToolBuilder->Targets = &ToolInputObjects;
 	RegisterTool(CommandInfos.BeginSeamTool, TEXT("UVSeamTool"), UVEditorSeamToolBuilder);
+
+	UUVEditorRecomputeUVsToolBuilder* UVEditorRecomputeUVsToolBuilder = NewObject<UUVEditorRecomputeUVsToolBuilder>();
+	UVEditorRecomputeUVsToolBuilder->Targets = &ToolInputObjects;
+	RegisterTool(CommandInfos.BeginRecomputeUVsTool, TEXT("UVRecomputeUVsTool"), UVEditorRecomputeUVsToolBuilder);
 }
 
 void UUVEditorMode::CreateToolkit()
diff --git a/Engine/Plugins/Experimental/UVEditor/Source/UVEditor/Private/UVEditorModeToolkit.cpp b/Engine/Plugins/Experimental/UVEditor/Source/UVEditor/Private/UVEditorModeToolkit.cpp
--- a/Engine/Plugins/Experimental/UVEditor/Source/UVEditor/Private/UVEditorModeToolkit.cpp
+++ b/Engine/Plugins/Experimental/UVEditor/Source/UVEditor/Private/UVEditorModeToolkit.cpp
@@ -141,6 +141,7 @@ void FUVEditorModeToolkit::Init(const TSharedPtr<IToolkitHost>& InitToolkitHost,
 	ToolbarBuilder.AddToolBarButton(CommandInfos.BeginParameterizeMeshTool);
 	ToolbarBuilder.AddToolBarButton(CommandInfos.BeginChannelEditTool);
 	ToolbarBuilder.AddToolBarButton(CommandInfos.BeginSeamTool);
+	ToolbarBuilder.AddToolBarButton(CommandInfos.BeginRecomputeUVsTool);
 
 	// Hook in the tool palette
 	ToolButtonsContainer->SetContent(ToolbarBuilder.MakeWidget());
diff --git a/Engine/Plugins/Experimental/UVEditor/Source/UVEditor/Public/UVEditorCommands.h b/Engine/Plugins/Experimental/UVEditor/Source/UVEditor/Public/UVEditorCommands.h
--- a/Engine/Plugins/Experimental/UVEditor/Source/UVEditor/Public/UVEditorCommands.h
+++ b/Engine/Plugins/Experimental/UVEditor/Source/UVEditor/Public/UVEditorCommands.h
@@ -23,6 +23,7 @@ public:
 	TSharedPtr<FUICommandInfo> BeginParameterizeMeshTool;
 	TSharedPtr<FUICommandInfo> BeginChannelEditTool;
 	TSharedPtr<FUICommandInfo> BeginSeamTool;
+	TSharedPtr<FUICommandInfo> BeginRecomputeUVsTool;
 
 	TSharedPtr<FUICommandInfo> AcceptOrCompleteActiveTool;
 	TSharedPtr<FUICommandInfo> CancelOrCompleteActiveTool;
diff --git a/Engine/Plugins/Experimental/UVEditor/Source/UVEditorTools/Private/UVEditorRecomputeUVsTool.cpp b/Engine/Plugins/Experimental/UVEditor/Source/UVEditorTools/Private/UVEditorRecomputeUVsTool.cpp
new file mode 100644
--- /dev/null
+++ b/Engine/Plugins/Experimental/UVEditor/Source/UVEditorTools/Private/UVEditorRecomputeUVsTool.cpp
@@ -0,0 +1,213 @@
+// Copyright Epic Games, Inc. All Rights Reserved.
+
+#include "UVEditorRecomputeUVsTool.h"
+#include "InteractiveToolManager.h"
+#include "ToolBuilderUtil.h"
+#include "DynamicMesh/DynamicMesh3.h"
+#include "Polygroups/PolygroupUtil.h"
+#include "FaceGroupUtil.h"
+#include "ToolSetupUtil.h"
+#include "ModelingToolTargetUtil.h"
+#include "ParameterizationOps/RecomputeUVsOp.h"
+#include "ToolTargets/UVEditorToolMeshInput.h"
+#include "Properties/MeshMaterialProperties.h"
+#include "UVToolContextObjects.h"
+#include "ContextObjectStore.h"
+
+
+#include "ExplicitUseGeometryMathTypes.h"		// using UE::Geometry::(math types)
+using namespace UE::Geometry;
+
+#define LOCTEXT_NAMESPACE "UUVEditorRecomputeUVsTool"
+
+
+/*
+ * ToolBuilder
+ */
+
+const FToolTargetTypeRequirements& UUVEditorRecomputeUVsToolBuilder::GetTargetRequirements() const
+{
+	static FToolTargetTypeRequirements TypeRequirements(UUVEditorToolMeshInput::StaticClass());
+	return TypeRequirements;
+}
+
+bool UUVEditorRecomputeUVsToolBuilder::CanBuildTool(const FToolBuilderState& SceneState) const
+{
+	return Targets && Targets->Num() > 0;
+}
+
+UInteractiveTool* UUVEditorRecomputeUVsToolBuilder::BuildTool(const FToolBuilderState& SceneState) const
+{
+	UUVEditorRecomputeUVsTool* NewTool = NewObject<UUVEditorRecomputeUVsTool>(SceneState.ToolManager);
+	NewTool->SetTargets(*Targets);
+
+	return NewTool;
+}
+
+
+/*
+ * Tool
+ */
+
+
+void UUVEditorRecomputeUVsTool::Setup()
+{
+	UInteractiveTool::Setup();
+
+	// initialize our properties
+
+	Settings = NewObject<URecomputeUVsToolProperties>(this);
+	Settings->RestoreProperties(this);
+	AddToolPropertySource(Settings);
+
+	if (Targets.Num() == 1)
+	{
+		PolygroupLayerProperties = NewObject<UPolygroupLayersProperties>(this);
+		PolygroupLayerProperties->RestoreProperties(this, TEXT("UVEditorRecomputeUVsTool"));
+		PolygroupLayerProperties->InitializeGroupLayers(Targets[0]->AppliedCanonical.Get());
+		PolygroupLayerProperties->WatchProperty(PolygroupLayerProperties->ActiveGroupLayer, [&](FName) { OnSelectedGroupLayerChanged(); });
+		AddToolPropertySource(PolygroupLayerProperties);
+		UpdateActiveGroupLayer();
+	}
+
+	Factories.SetNum(Targets.Num());
+	for (int32 TargetIndex = 0; TargetIndex < Targets.Num(); ++TargetIndex)
+	{
+		TObjectPtr<UUVEditorToolMeshInput> Target = Targets[TargetIndex];
+		Factories[TargetIndex] = NewObject<URecomputeUVsOpFactory>();
+		Factories[TargetIndex]->TargetTransform = (UE::Geometry::FTransform3d)Target->AppliedPreview->PreviewMesh->GetTransform();
+		Factories[TargetIndex]->Settings = Settings;
+		Factories[TargetIndex]->OriginalMesh = Target->AppliedCanonical;
+		Factories[TargetIndex]->InputGroups = ActiveGroupSet;
+		Factories[TargetIndex]->GetSelectedUVChannel = [Target]() { return Target->UVLayerIndex; };
+
+		Target->AppliedPreview->ChangeOpFactory(Factories[TargetIndex]);
+		Target->AppliedPreview->OnMeshUpdated.AddWeakLambda(this, [Target](UMeshOpPreviewWithBackgroundCompute* Preview) {
+			Target->UpdateUnwrapPreviewFromAppliedPreview();
+			});
+
+		Target->AppliedPreview->InvalidateResult();
+	}
+
+	SetToolDisplayName(LOCTEXT("ToolNameLocal", "UV Unwrap"));
+	GetToolManager()->DisplayMessage(
+		LOCTEXT("OnStartTool_Regions", "Generate UVs for Polygroups or existing UV charts of the Mesh using various strategies."),
+		EToolMessageLevel::UserNotification);
+}
+
+
+void UUVEditorRecomputeUVsTool::OnPropertyModified(UObject* PropertySet, FProperty* Property)
+{
+	bool bForceMaterialUpdate = false;
+	if (PropertySet == Settings )
+	{
+		// One of the UV generation properties must have changed.  Dirty the result to force a recompute
+		for (TObjectPtr<UUVEditorToolMeshInput> Target : Targets)
+		{
+			Target->AppliedPreview->InvalidateResult();
+		}
+		bForceMaterialUpdate = true;
+	}
+
+}
+
+
+void UUVEditorRecomputeUVsTool::Shutdown(EToolShutdownType ShutdownType)
+{
+	Settings->SaveProperties(this);
+	if (PolygroupLayerProperties) {
+		PolygroupLayerProperties->RestoreProperties(this, TEXT("UVEditorRecomputeUVsTool"));
+	}
+
+	if (ShutdownType == EToolShutdownType::Accept)
+	{
+		UUVToolEmitChangeAPI* ChangeAPI = GetToolManager()->GetContextObjectStore()->FindContext<UUVToolEmitChangeAPI>();
+
+		const FText TransactionName(LOCTEXT("RecomputeUVsTransactionName", "Recompute UVs"));
+		ChangeAPI->BeginUndoTransaction(TransactionName);
+		for (TObjectPtr<UUVEditorToolMeshInput> Target : Targets)
+		{
+			// Set things up for undo. 
+			FDynamicMeshChangeTracker ChangeTracker(Target->UnwrapCanonical.Get());
+			ChangeTracker.BeginChange();
+
+			for (int32 Tid : Target->UnwrapCanonical->TriangleIndicesItr())
+			{
+				ChangeTracker.SaveTriangle(Tid, true);
+			}
+
+			Target->UpdateCanonicalFromPreviews();
+
+			ChangeAPI->EmitToolIndependentUnwrapCanonicalChange(Target, ChangeTracker.EndChange(), LOCTEXT("ApplyRecomputeUVsTool", "Unwrap Tool"));
+		}
+
+		ChangeAPI->EndUndoTransaction();
+	}
+	else
+	{
+		// Reset the inputs
+		for (TObjectPtr<UUVEditorToolMeshInput> Target : Targets)
+		{
+			Target->UpdatePreviewsFromCanonical();
+		}
+	}
+
+	for (TObjectPtr<UUVEditorToolMeshInput> Target : Targets)
+	{
+		Target->AppliedPreview->ClearOpFactory();
+	}
+	for (int32 TargetIndex = 0; TargetIndex < Targets.Num(); ++TargetIndex)
+	{
+		Factories[TargetIndex] = nullptr;
+	}
+	Settings = nullptr;
+	Targets.Empty();
+}
+
+void UUVEditorRecomputeUVsTool::Render(IToolsContextRenderAPI* RenderAPI)
+{
+}
+
+
+
+bool UUVEditorRecomputeUVsTool::CanAccept() const
+{
+	bool bPreviewsHaveValidResults = true;
+	for (TObjectPtr<UUVEditorToolMeshInput> Target : Targets)
+	{
+		bPreviewsHaveValidResults = bPreviewsHaveValidResults && Target->AppliedPreview->HaveValidResult();
+	}
+	return bPreviewsHaveValidResults;
+}
+
+
+void UUVEditorRecomputeUVsTool::OnSelectedGroupLayerChanged()
+{
+	UpdateActiveGroupLayer();
+	for (TObjectPtr<UUVEditorToolMeshInput> Target : Targets)
+	{
+		Target->AppliedPreview->InvalidateResult();
+	}
+}
+
+
+void UUVEditorRecomputeUVsTool::UpdateActiveGroupLayer()
+{
+	if (Targets.Num() == 1)
+	{
+		if (PolygroupLayerProperties->HasSelectedPolygroup() == false)
+		{
+			ActiveGroupSet = MakeShared<UE::Geometry::FPolygroupSet, ESPMode::ThreadSafe>(Targets[0]->AppliedCanonical.Get());
+		}
+		else
+		{
+			FName SelectedName = PolygroupLayerProperties->ActiveGroupLayer;
+			FDynamicMeshPolygroupAttribute* FoundAttrib = UE::Geometry::FindPolygroupLayerByName(*Targets[0]->AppliedCanonical, SelectedName);
+			ensureMsgf(FoundAttrib, TEXT("Selected Attribute Not Found! Falling back to Default group layer."));
+			ActiveGroupSet = MakeShared<UE::Geometry::FPolygroupSet, ESPMode::ThreadSafe>(Targets[0]->AppliedCanonical.Get(), FoundAttrib);
+		}
+	}
+}
+
+
+#undef LOCTEXT_NAMESPACE
diff --git a/Engine/Plugins/Experimental/UVEditor/Source/UVEditorTools/Public/UVEditorRecomputeUVsTool.h b/Engine/Plugins/Experimental/UVEditor/Source/UVEditorTools/Public/UVEditorRecomputeUVsTool.h
new file mode 100644
--- /dev/null
+++ b/Engine/Plugins/Experimental/UVEditor/Source/UVEditorTools/Public/UVEditorRecomputeUVsTool.h
@@ -0,0 +1,90 @@
+// Copyright Epic Games, Inc. All Rights Reserved.
+
+#pragma once
+
+#include "CoreMinimal.h"
+#include "BaseTools/SingleSelectionMeshEditingTool.h"
+#include "DynamicMesh/DynamicMesh3.h"
+#include "MeshOpPreviewHelpers.h"
+#include "PropertySets/PolygroupLayersProperties.h"
+#include "Properties/RecomputeUVsProperties.h"
+#include "Polygroups/PolygroupSet.h"
+#include "Drawing/UVLayoutPreview.h"
+
+#include "UVEditorRecomputeUVsTool.generated.h"
+
+
+// predeclarations
+class UDynamicMeshComponent;
+class UMaterialInterface;
+class UMaterialInstanceDynamic;
+class UUVEditorToolMeshInput;
+class URecomputeUVsOpFactory;
+
+/**
+ *
+ */
+UCLASS()
+class UVEDITORTOOLS_API UUVEditorRecomputeUVsToolBuilder : public USingleSelectionMeshEditingToolBuilder
+{
+	GENERATED_BODY()
+
+public:
+	virtual bool CanBuildTool(const FToolBuilderState& SceneState) const override;
+	virtual UInteractiveTool* BuildTool(const FToolBuilderState& SceneState) const override;
+
+	// This is a pointer so that it can be updated under the builder without
+	// having to set it in the mode after initializing targets.
+	const TArray<TObjectPtr<UUVEditorToolMeshInput>>* Targets = nullptr;
+
+protected:
+	virtual const FToolTargetTypeRequirements& GetTargetRequirements() const override;
+};
+
+
+/**
+ * UUVEditorRecomputeUVsTool Recomputes UVs based on existing segmentations of the mesh
+ */
+UCLASS()
+class UVEDITORTOOLS_API UUVEditorRecomputeUVsTool : public UInteractiveTool 
+{
+	GENERATED_BODY()
+
+public:
+	virtual void Setup() override;
+	virtual void Shutdown(EToolShutdownType ShutdownType) override;
+
+	virtual void Render(IToolsContextRenderAPI* RenderAPI) override;
+
+	virtual bool HasCancel() const override { return true; }
+	virtual bool HasAccept() const override { return true; }
+	virtual bool CanAccept() const override;
+
+	virtual void OnPropertyModified(UObject* PropertySet, FProperty* Property) override;
+
+
+	/**
+     * The tool will operate on the meshes given here.
+     */
+	virtual void SetTargets(const TArray<TObjectPtr<UUVEditorToolMeshInput>>& TargetsIn)
+	{
+		Targets = TargetsIn;
+	}
+
+protected:
+	UPROPERTY()
+	TArray<TObjectPtr<UUVEditorToolMeshInput>> Targets;
+
+	UPROPERTY()
+	TObjectPtr<URecomputeUVsToolProperties> Settings = nullptr;
+
+	UPROPERTY()
+	TObjectPtr<UPolygroupLayersProperties> PolygroupLayerProperties = nullptr;
+
+	UPROPERTY()
+	TArray<TObjectPtr<URecomputeUVsOpFactory>> Factories;
+
+	TSharedPtr<UE::Geometry::FPolygroupSet, ESPMode::ThreadSafe> ActiveGroupSet;
+	void OnSelectedGroupLayerChanged();
+	void UpdateActiveGroupLayer();
+};
diff --git a/Engine/Plugins/Runtime/MeshModelingToolset/Source/ModelingOperators/Private/ParameterizationOps/RecomputeUVsOp.cpp b/Engine/Plugins/Runtime/MeshModelingToolset/Source/ModelingOperators/Private/ParameterizationOps/RecomputeUVsOp.cpp
--- a/Engine/Plugins/Runtime/MeshModelingToolset/Source/ModelingOperators/Private/ParameterizationOps/RecomputeUVsOp.cpp
+++ b/Engine/Plugins/Runtime/MeshModelingToolset/Source/ModelingOperators/Private/ParameterizationOps/RecomputeUVsOp.cpp
@@ -8,6 +8,7 @@
 #include "Parameterization/MeshLocalParam.h"
 #include "Parameterization/DynamicMeshUVEditor.h"
 #include "Parameterization/PatchBasedMeshUVGenerator.h"
+#include "Properties/RecomputeUVsProperties.h"
 
 #include "DynamicSubmesh3.h"
 
@@ -377,5 +378,63 @@ bool FRecomputeUVsOp::CalculateResult_RegionOptimization(FProgressCancel* Progre
 }
 
 
+/**
+ * Factory
+ */
+
+
+TUniquePtr<FDynamicMeshOperator> URecomputeUVsOpFactory::MakeNewOperator()
+{
+	FAxisAlignedBox3d MeshBounds = OriginalMesh->GetBounds();
+	TUniquePtr<FRecomputeUVsOp> RecomputeUVsOp = MakeUnique<FRecomputeUVsOp>();
+	RecomputeUVsOp->InputMesh = OriginalMesh;
+	RecomputeUVsOp->InputGroups = InputGroups;
+	RecomputeUVsOp->UVLayer = GetSelectedUVChannel();
+
+	RecomputeUVsOp->IslandMode = (UE::Geometry::ERecomputeUVsIslandMode)(int)Settings->IslandMode;
+	RecomputeUVsOp->UnwrapType = (UE::Geometry::ERecomputeUVsUnwrapType)(int)Settings->UnwrapType;
+
+	RecomputeUVsOp->bAutoRotate = (Settings->AutoRotation == ERecomputeUVsToolOrientationMode::MinBoxBounds);
+
+	RecomputeUVsOp->NormalSmoothingRounds = Settings->SmoothingSteps;
+	RecomputeUVsOp->NormalSmoothingAlpha = Settings->SmoothingAlpha;
+
+	RecomputeUVsOp->bMergingOptimization = Settings->bIslandMerging;
+	RecomputeUVsOp->MergingThreshold = Settings->MergingThreshold;
+	RecomputeUVsOp->MaxNormalDeviationDeg = Settings->MaxAngleDeviation;
+
+	RecomputeUVsOp->bPackUVs = Settings->bAutoPack;
+	if (Settings->bAutoPack)
+	{
+		RecomputeUVsOp->PackingTextureResolution = Settings->TextureResolution;
+		RecomputeUVsOp->bNormalizeAreas = false;
+		RecomputeUVsOp->AreaScaling = 1.0;
+	}
+	else
+	{
+		switch (Settings->UVScaleMode)
+		{
+		case ERecomputeUVsToolUVScaleMode::NoScaling:
+			RecomputeUVsOp->bNormalizeAreas = false;
+			RecomputeUVsOp->AreaScaling = 1.0;
+			break;
+		case ERecomputeUVsToolUVScaleMode::NormalizeToBounds:
+			RecomputeUVsOp->bNormalizeAreas = true;
+			RecomputeUVsOp->AreaScaling = Settings->UVScale / MeshBounds.MaxDim();
+			break;
+		case ERecomputeUVsToolUVScaleMode::NormalizeToWorld:
+			RecomputeUVsOp->bNormalizeAreas = true;
+			RecomputeUVsOp->AreaScaling = Settings->UVScale;
+			break;
+		}
+	}
+
+
+	RecomputeUVsOp->SetTransform(TargetTransform);
+
+	return RecomputeUVsOp;
+}
+
+
 
 #undef LOCTEXT_NAMESPACE
\ No newline at end of file
diff --git a/Engine/Plugins/Runtime/MeshModelingToolset/Source/ModelingOperators/Public/ParameterizationOps/RecomputeUVsOp.h b/Engine/Plugins/Runtime/MeshModelingToolset/Source/ModelingOperators/Public/ParameterizationOps/RecomputeUVsOp.h
--- a/Engine/Plugins/Runtime/MeshModelingToolset/Source/ModelingOperators/Public/ParameterizationOps/RecomputeUVsOp.h
+++ b/Engine/Plugins/Runtime/MeshModelingToolset/Source/ModelingOperators/Public/ParameterizationOps/RecomputeUVsOp.h
@@ -8,6 +8,9 @@
 #include "DynamicMesh/DynamicMeshAttributeSet.h"
 #include "Polygroups/PolygroupSet.h"
 
+#include "RecomputeUVsOp.generated.h"
+
+class URecomputeUVsToolProperties;
 
 namespace UE
 {
@@ -102,3 +105,29 @@ protected:
 
 } // end namespace UE::Geometry
 } // end namespace UE
+
+
+/**
+ * Can be hooked up to a UMeshOpPreviewWithBackgroundCompute to perform UV solving operations.
+ * 
+ * Inherits from UObject so that it can hold a strong pointer to the settings UObject, which
+ * needs to be a UObject to be displayed in the details panel.
+ */
+UCLASS()
+class MODELINGOPERATORS_API URecomputeUVsOpFactory : public UObject, public UE::Geometry::IDynamicMeshOperatorFactory
+{
+	GENERATED_BODY()
+
+public:
+
+	// IDynamicMeshOperatorFactory API
+	virtual TUniquePtr<UE::Geometry::FDynamicMeshOperator> MakeNewOperator() override;
+
+	UPROPERTY()
+	TObjectPtr<URecomputeUVsToolProperties> Settings = nullptr;
+
+	TSharedPtr<UE::Geometry::FPolygroupSet, ESPMode::ThreadSafe> InputGroups;
+	TSharedPtr<UE::Geometry::FDynamicMesh3, ESPMode::ThreadSafe> OriginalMesh;
+	TUniqueFunction<int32()> GetSelectedUVChannel = []() { return 0; };
+	UE::Geometry::FTransform3d TargetTransform;
+};
\ No newline at end of file
diff --git a/Engine/Plugins/Runtime/MeshModelingToolset/Source/ModelingOperators/Public/Properties/RecomputeUVsProperties.h b/Engine/Plugins/Runtime/MeshModelingToolset/Source/ModelingOperators/Public/Properties/RecomputeUVsProperties.h
new file mode 100644
--- /dev/null
+++ b/Engine/Plugins/Runtime/MeshModelingToolset/Source/ModelingOperators/Public/Properties/RecomputeUVsProperties.h
@@ -0,0 +1,113 @@
+// Copyright Epic Games, Inc. All Rights Reserved.
+
+#pragma once
+
+#include "CoreMinimal.h"
+
+#include "InteractiveTool.h"
+
+#include "RecomputeUVsProperties.generated.h"
+
+
+UENUM()
+enum class ERecomputeUVsPropertiesUnwrapType
+{
+	// Values must match UE::Geometry::ERecomputeUVsUnwrapType
+
+	/** ExpMap UV Flattening Algorithm is very fast but has limited ability to control stretching/distortion */
+	ExpMap = 0,
+	/** Conformal UV Flattening algorithm is increasingly-expensive on large islands but better at distortion control */
+	Conformal = 1
+};
+
+
+
+UENUM()
+enum class ERecomputeUVsPropertiesIslandMode
+{
+	// Values must match UE::Geometry::ERecomputeUVsIslandMode
+
+	/** Use Active Mesh Polygroups Layer to define initial UV Islands */
+	PolyGroups = 0,
+	/** Use Existing UV Layer to define UV Islands (ie re-solve UV flattening based on existing UVs) */
+	ExistingUVs = 1
+};
+
+
+UENUM()
+enum class ERecomputeUVsToolOrientationMode
+{
+	/**  */
+	None,
+	/**  */
+	MinBoxBounds
+};
+
+
+UENUM()
+enum class ERecomputeUVsToolUVScaleMode
+{
+	/** No scaling is applied to UV islands */
+	NoScaling,
+	/** Scale UV islands such that they have constant relative area, relative to object bounds */
+	NormalizeToBounds,
+	/** Scale UV islands such that they have constant relative area, relative to world space */
+	NormalizeToWorld
+};
+
+
+UCLASS()
+class MODELINGOPERATORS_API URecomputeUVsToolProperties : public UInteractiveToolPropertySet
+{
+	GENERATED_BODY()
+public:
+
+	/** IslandMode determines which patches of triangles will be used as initial UV Islands */
+	UPROPERTY(EditAnywhere, Category = "UV Generation")
+	ERecomputeUVsPropertiesIslandMode IslandMode = ERecomputeUVsPropertiesIslandMode::ExistingUVs;
+
+	/** Which type of UV flattening algorithm to use */
+	UPROPERTY(EditAnywhere, Category = "UV Generation", meta = (EditCondition = "bIslandMerging == false"))
+	ERecomputeUVsPropertiesUnwrapType UnwrapType = ERecomputeUVsPropertiesUnwrapType::Conformal;
+
+	/** Which type of automatic island orientation to use */
+	UPROPERTY(EditAnywhere, Category = "UV Generation")
+	ERecomputeUVsToolOrientationMode AutoRotation = ERecomputeUVsToolOrientationMode::MinBoxBounds;
+
+	/** If Island Merging is enabled, then the initial UV islands will be merged into larger islands if it does not increase distortion/stretching beyond the limits below */
+	UPROPERTY(EditAnywhere, Category = "UV Generation")
+	bool bIslandMerging = false;
+
+	/** Distortion/Stretching Threshold for island merging - larger values increase the allowable UV stretching */
+	UPROPERTY(EditAnywhere, Category = "UV Generation", meta = (UIMin = "1.0", UIMax = "5.0", ClampMin = "1.0", EditCondition = "bIslandMerging == true"))
+	float MergingThreshold = 1.5f;
+
+	/** UV islands will not be merged if their average face normals deviate by larger than this amount */
+	UPROPERTY(EditAnywhere, Category = "UV Generation", meta = (UIMin = "0.0", UIMax = "90.0", ClampMin = "0.0", ClampMax = "180.0", EditCondition = "bIslandMerging == true"))
+	float MaxAngleDeviation = 45.0f;
+
+	/** Number of smoothing steps to apply in the ExpMap UV Generation method (Smoothing slightly increases distortion but produces more stable results) */
+	UPROPERTY(EditAnywhere, Category = "ExpMap UV Generator", meta = (UIMin = "0", UIMax = "25", ClampMin = "0", ClampMax = "1000", EditCondition = "UnwrapType == ERecomputeUVsUnwrapType::ExpMap || bIslandMerging == true"))
+	int SmoothingSteps = 5;
+
+	/** Smoothing parameter, larger values result in faster smoothing in each step */
+	UPROPERTY(EditAnywhere, Category = "ExpMap UV Generator", meta = (UIMin = "0", UIMax = "1.0", ClampMin = "0", ClampMax = "1.0", EditCondition = "UnwrapType == ERecomputeUVsUnwrapType::ExpMap || bIslandMerging == true"))
+	float SmoothingAlpha = 0.25f;
+
+	/** If enabled, result UVs are automatically packed into the standard UV 0-1 square */
+	UPROPERTY(EditAnywhere, Category = "UV Layout")
+	bool bAutoPack = true;
+
+	/** Target texture resolution used for UV packing, which determines gutter size */
+	UPROPERTY(EditAnywhere, Category = "UV Layout", meta = (UIMin = "64", UIMax = "2048", ClampMin = "2", ClampMax = "4096", EditCondition = bAutoPack))
+	int TextureResolution = 1024;
+
+	/** Scaling applied to UV islands after solve */
+	UPROPERTY(EditAnywhere, Category = "Transform UVs", meta = (EditCondition = "bAutoPack == false"))
+	ERecomputeUVsToolUVScaleMode UVScaleMode = ERecomputeUVsToolUVScaleMode::NormalizeToBounds;
+
+	/** Scaling factor used for UV island normalization/scaling */
+	UPROPERTY(EditAnywhere, Category = "Transform UVs", meta = (EditCondition = "UVScaleMode!=ERecomputeUVsToolUVScaleMode::NoScaling && bAutoPack == false", UIMin = "0.001", UIMax = "10", ClampMin = "0.00001", ClampMax = "1000000.0"))
+	float UVScale = 1.0;
+};
+
