UVEditor: Change how background imagery is handled to display active image as a material on the applied preview, adds a new standalone color checkerboard pattern, and cleans up the toolbar controls for display settings a little bit.

#rnx

diff --git a/Engine/Plugins/Experimental/UVEditor/Source/UVEditor/Private/UVEditorBackgroundPreview.cpp b/Engine/Plugins/Experimental/UVEditor/Source/UVEditor/Private/UVEditorBackgroundPreview.cpp
--- a/Engine/Plugins/Experimental/UVEditor/Source/UVEditor/Private/UVEditorBackgroundPreview.cpp
+++ b/Engine/Plugins/Experimental/UVEditor/Source/UVEditor/Private/UVEditorBackgroundPreview.cpp
@@ -14,7 +14,7 @@ void UUVEditorBackgroundPreview::OnCreated()
 {
 	Settings = NewObject<UUVEditorBackgroundPreviewProperties>(this);
 	Settings->WatchProperty(Settings->bVisible, [this](bool) { bSettingsModified = true; });
-	Settings->WatchProperty(Settings->bUseMaterials, [this](bool) { bSettingsModified = true; });
+	Settings->WatchProperty(Settings->SourceType, [this](EUVEditorBackgroundSourceType) { bSettingsModified = true; });
 	Settings->WatchProperty(Settings->SourceTexture, [this](UTexture2D*) {bSettingsModified = true; });
 	Settings->WatchProperty(Settings->SourceMaterial, [this](UMaterial*) {bSettingsModified = true; });
 
@@ -60,21 +60,47 @@ void UUVEditorBackgroundPreview::UpdateBackground()
 	UMaterial* Material = LoadObject<UMaterial>(nullptr, TEXT("/UVEditor/Materials/UVEditorBackground"));
 	check(Material);	
 	BackgroundMaterial = UMaterialInstanceDynamic::Create(Material, this);
-	if (Settings->bUseMaterials)
+	switch (Settings->SourceType)
 	{
-		if (Settings->SourceMaterial)
+		case EUVEditorBackgroundSourceType::Checkerboard:
 		{
-			BackgroundMaterial = UMaterialInstanceDynamic::Create(Settings->SourceMaterial.Get(), this);
+			// Do nothing, since the default material is already set up for a checkerboard.
 		}
-	}
-	else
-	{
-		if (Settings->SourceTexture)
+		break;
+
+		case EUVEditorBackgroundSourceType::Material:
+		{
+			if (Settings->SourceMaterial)
+			{
+				BackgroundMaterial = UMaterialInstanceDynamic::Create(Settings->SourceMaterial.Get(), this);
+			}
+		}
+		break;
+
+		case EUVEditorBackgroundSourceType::Texture:
 		{
-			BackgroundMaterial->SetTextureParameterValue(TEXT("BackgroundBaseMap"), Settings->SourceTexture);
+			if (Settings->SourceTexture)
+			{
+				if (Settings->SourceTexture->IsCurrentlyVirtualTextured())
+				{
+					BackgroundMaterial->SetTextureParameterValue(TEXT("BackgroundVTBaseMap"), Settings->SourceTexture);
+					BackgroundMaterial->SetScalarParameterValue(TEXT("BackgroundVirtualTextureSwitch"), 1);
+				}
+				else
+				{
+					BackgroundMaterial->SetTextureParameterValue(TEXT("BackgroundBaseMap"), Settings->SourceTexture);
+					BackgroundMaterial->SetScalarParameterValue(TEXT("BackgroundVirtualTextureSwitch"), 0);
+				}
+			}
 		}
+		break;
+
+		default:
+			ensure(false);
 	}
+
 	BackgroundMaterial->SetScalarParameterValue(TEXT("BackgroundPixelDepthOffset"), 0);
+	OnBackgroundMaterialChange.Broadcast(BackgroundMaterial);
 	BackgroundComponent->Clear();
 
 	for (int32 GridStepX = 0; GridStepX < GridCellCountX; ++GridStepX)
diff --git a/Engine/Plugins/Experimental/UVEditor/Source/UVEditor/Private/UVEditorMode.cpp b/Engine/Plugins/Experimental/UVEditor/Source/UVEditor/Private/UVEditorMode.cpp
--- a/Engine/Plugins/Experimental/UVEditor/Source/UVEditor/Private/UVEditorMode.cpp
+++ b/Engine/Plugins/Experimental/UVEditor/Source/UVEditor/Private/UVEditorMode.cpp
@@ -140,6 +140,11 @@ void UUVEditorMode::Enter()
 	BackgroundVisualization->Settings->WatchProperty(BackgroundVisualization->Settings->bVisible, 
 		[this](bool IsVisible) { UpdateTriangleMaterialBasedOnBackground(IsVisible); });
 
+	BackgroundVisualization->OnBackgroundMaterialChange.AddWeakLambda(this,
+		[this](TObjectPtr<UMaterialInstanceDynamic> MaterialInstance) {
+			UpdatePreviewMaterialBasedOnBackground();
+		});
+
 	PropertyObjectsToTick.Add(BackgroundVisualization->Settings);
 
 	RegisterTools();
@@ -584,6 +589,19 @@ void UUVEditorMode::UpdateTriangleMaterialBasedOnBackground(bool IsBackgroundVis
 	}
 }
 
+void UUVEditorMode::UpdatePreviewMaterialBasedOnBackground()
+{
+	for (int32 AssetID = 0; AssetID < ToolInputObjects.Num(); ++AssetID) {
+		ToolInputObjects[AssetID]->AppliedPreview->OverrideMaterial = nullptr;
+	}
+	if (BackgroundVisualization->Settings->bVisible)
+	{
+		for (int32 AssetID = 0; AssetID < ToolInputObjects.Num(); ++AssetID) {
+			ToolInputObjects[AssetID]->AppliedPreview->OverrideMaterial = BackgroundVisualization->BackgroundMaterial;
+		}
+	}
+}
+
 void UUVEditorMode::ModeTick(float DeltaTime)
 {
 	using namespace UVEditorModeLocals;
diff --git a/Engine/Plugins/Experimental/UVEditor/Source/UVEditor/Private/UVEditorModeToolkit.cpp b/Engine/Plugins/Experimental/UVEditor/Source/UVEditor/Private/UVEditorModeToolkit.cpp
--- a/Engine/Plugins/Experimental/UVEditor/Source/UVEditor/Private/UVEditorModeToolkit.cpp
+++ b/Engine/Plugins/Experimental/UVEditor/Source/UVEditor/Private/UVEditorModeToolkit.cpp
@@ -312,19 +312,10 @@ TSharedRef<SWidget> FUVEditorModeToolkit::CreateBackgroundSettingsWidget()
 		.HAlign(HAlign_Fill)
 		.Padding(4)
 		[
-			SNew(SVerticalBox)
-
-			+ SVerticalBox::Slot()
-				.AutoHeight()
-				.Padding(2, 0, 0, 0)
+			SNew(SBox)			
+				.MinDesiredWidth(500)				
 			[
-				SNew(SVerticalBox)
-
-				+ SVerticalBox::Slot()
-				.AutoHeight()
-				[
-					BackgroundDetailsContainer
-				]
+				BackgroundDetailsContainer				
 			]
 		];
 
diff --git a/Engine/Plugins/Experimental/UVEditor/Source/UVEditor/Private/UVEditorToolkit.cpp b/Engine/Plugins/Experimental/UVEditor/Source/UVEditor/Private/UVEditorToolkit.cpp
--- a/Engine/Plugins/Experimental/UVEditor/Source/UVEditor/Private/UVEditorToolkit.cpp
+++ b/Engine/Plugins/Experimental/UVEditor/Source/UVEditor/Private/UVEditorToolkit.cpp
@@ -448,8 +448,8 @@ void FUVEditorToolkit::PostInitAssetEditor()
 		{
 			return UVModeToolkit->CreateBackgroundSettingsWidget();
 		}),
-		LOCTEXT("UVEditorBackgroundSettings_Label", "Background"),
-		LOCTEXT("UVEditorBackgroundSettings_ToolTip", "Change the visibility and other settings of the background."),
+		LOCTEXT("UVEditorBackgroundSettings_Label", "Display"),
+		LOCTEXT("UVEditorBackgroundSettings_ToolTip", "Change the background and other viewport display settings."),
 		FSlateIcon(FUVEditorStyle::Get().GetStyleSetName(), "UVEditor.BackgroundSettings")
 	));
 
diff --git a/Engine/Plugins/Experimental/UVEditor/Source/UVEditor/Public/UVEditorBackgroundPreview.h b/Engine/Plugins/Experimental/UVEditor/Source/UVEditor/Public/UVEditorBackgroundPreview.h
--- a/Engine/Plugins/Experimental/UVEditor/Source/UVEditor/Public/UVEditorBackgroundPreview.h
+++ b/Engine/Plugins/Experimental/UVEditor/Source/UVEditor/Public/UVEditorBackgroundPreview.h
@@ -9,6 +9,18 @@
 #include "GeometryBase.h"
 #include "UVEditorBackgroundPreview.generated.h"
 
+
+/**
+ * Enum to control the background visualiztion mode
+ */
+UENUM()
+enum class EUVEditorBackgroundSourceType
+{
+	Checkerboard,
+	Texture,
+	Material
+};
+
 /**
  * Visualization settings for the UUVEditorBackgroundPreview
  */
@@ -18,19 +30,19 @@ class UVEDITOR_API UUVEditorBackgroundPreviewProperties : public UInteractiveToo
 	GENERATED_BODY()
 public:
 	/** Should the background be shown */
-	UPROPERTY(EditAnywhere, Category = Background)
+	UPROPERTY(EditAnywhere, Category = Background, meta = (DisplayName = "Display Background") )
 	bool bVisible = false;
 
-	/** Should the background show textures or materials */
-	UPROPERTY(EditAnywhere, Category = Background, meta = (EditCondition = "bVisible"))
-	bool bUseMaterials = false;
+	/** Source of background visuals */
+	UPROPERTY(EditAnywhere, Category = Background, meta = (EditCondition = "bVisible", DisplayName = "Background Source"))
+	EUVEditorBackgroundSourceType SourceType = EUVEditorBackgroundSourceType::Checkerboard;
 
-	/** Used when bUseMaterials is false */
-	UPROPERTY(EditAnywhere, Category = Background, meta = (EditCondition = "!bUseMaterials && bVisible"))
+	/** Display a background based on the selected texture */
+	UPROPERTY(EditAnywhere, Category = Background, meta = (EditCondition = "SourceType==EUVEditorBackgroundSourceType::Texture && bVisible", EditConditionHides = true))
 	TObjectPtr<UTexture2D> SourceTexture;
 
-	/** Used when bUseMaterials is true */
-	UPROPERTY(EditAnywhere, Category = Background, meta = (EditCondition = "bUseMaterials && bVisible"))
+	/** Display a background based on the selected material */
+	UPROPERTY(EditAnywhere, Category = Background, meta = (EditCondition = "SourceType==EUVEditorBackgroundSourceType::Material && bVisible", EditConditionHides = true))
 	TObjectPtr<UMaterial> SourceMaterial;
 };
 
@@ -64,6 +76,9 @@ public:
 	UPROPERTY()
 	TObjectPtr<UMaterialInstanceDynamic> BackgroundMaterial;
 
+	DECLARE_MULTICAST_DELEGATE_OneParam(FOnBackgroundMaterialChange, TObjectPtr<UMaterialInstanceDynamic> MaterialInstance);
+	FOnBackgroundMaterialChange OnBackgroundMaterialChange;
+
 protected:
 	virtual void OnCreated() override;
 
diff --git a/Engine/Plugins/Experimental/UVEditor/Source/UVEditor/Public/UVEditorMode.h b/Engine/Plugins/Experimental/UVEditor/Source/UVEditor/Public/UVEditorMode.h
--- a/Engine/Plugins/Experimental/UVEditor/Source/UVEditor/Public/UVEditorMode.h
+++ b/Engine/Plugins/Experimental/UVEditor/Source/UVEditor/Public/UVEditorMode.h
@@ -124,6 +124,7 @@ protected:
 	virtual void OnToolEnded(UInteractiveToolManager* Manager, UInteractiveTool* Tool) override {}
 	
 	void UpdateTriangleMaterialBasedOnBackground(bool IsBackgroundVisible);
+	void UpdatePreviewMaterialBasedOnBackground();
 
 	/**
 	 * Stores original input objects, for instance UStaticMesh pointers. AssetIDs on tool input 
