UVEditor: Correct bug with UVEditorRecomputeUVsToolBuiilder inheriting from the wrong class, causing NonUnity build errors.

#rnx

diff --git a/Engine/Plugins/Experimental/UVEditor/Source/UVEditorTools/Public/UVEditorRecomputeUVsTool.h b/Engine/Plugins/Experimental/UVEditor/Source/UVEditorTools/Public/UVEditorRecomputeUVsTool.h
--- a/Engine/Plugins/Experimental/UVEditor/Source/UVEditorTools/Public/UVEditorRecomputeUVsTool.h
+++ b/Engine/Plugins/Experimental/UVEditor/Source/UVEditorTools/Public/UVEditorRecomputeUVsTool.h
@@ -25,7 +25,7 @@ class URecomputeUVsOpFactory;
  *
  */
 UCLASS()
-class UVEDITORTOOLS_API UUVEditorRecomputeUVsToolBuilder : public USingleSelectionMeshEditingToolBuilder
+class UVEDITORTOOLS_API UUVEditorRecomputeUVsToolBuilder : public UInteractiveToolBuilder
 {
 	GENERATED_BODY()
 
