Add some tracing to the UV Editor call stack

#rnx

diff --git a/Engine/Plugins/Experimental/UVEditor/Source/UVEditorTools/Private/UVEditorSeamTool.cpp b/Engine/Plugins/Experimental/UVEditor/Source/UVEditorTools/Private/UVEditorSeamTool.cpp
--- a/Engine/Plugins/Experimental/UVEditor/Source/UVEditorTools/Private/UVEditorSeamTool.cpp
+++ b/Engine/Plugins/Experimental/UVEditor/Source/UVEditorTools/Private/UVEditorSeamTool.cpp
@@ -85,6 +85,8 @@ namespace UVEditorSeamToolLocals
 		UPointSetComponent* UnwrapPointSet, UPointSetComponent* AppliedPointSet,
 		int32 AppliedVid, const TArray<int32>& UnwrapVids, const FColor& Color)
 	{
+		TRACE_CPUPROFILER_EVENT_SCOPE(AddDisplayedPoints);
+
 		FTransform AppliedTransform = InputObject->AppliedPreview->PreviewMesh->GetTransform();
 		AppliedPointSet->AddPoint(FRenderablePoint(
 			InputObject->AppliedCanonical->GetVertex(AppliedVid), Color, PointSize));
@@ -100,6 +102,8 @@ namespace UVEditorSeamToolLocals
 		const TArray<int32>& NewPathVids, bool bPathIsFromUnwrap, 
 		const FColor& Color, float Thickness, float DepthBias)
 	{
+		TRACE_CPUPROFILER_EVENT_SCOPE(AddDisplayedPath);
+
 		FTransform AppliedTransform = InputObject->AppliedPreview->PreviewMesh->GetTransform();
 		if (NewPathVids.Num() < 2)
 		{
@@ -296,6 +300,8 @@ void UUVEditorSeamTool::SetTargets(const TArray<TObjectPtr<UUVEditorToolMeshInpu
 
 void UUVEditorSeamTool::Setup()
 {
+	TRACE_CPUPROFILER_EVENT_SCOPE(UVEditorSeamTool_Setup);
+
 	using namespace UVEditorSeamToolLocals;
 
 	UInteractiveTool::Setup();
@@ -474,6 +480,8 @@ void UUVEditorSeamTool::Setup()
 
 void UUVEditorSeamTool::ReconstructExistingSeamsVisualization()
 {
+	TRACE_CPUPROFILER_EVENT_SCOPE(UVEditorSeamTool_ReconstructExistingSeamsVisualization);
+
 	using namespace UVEditorSeamToolLocals;
 
 	ULineSetComponent* SeamLines = LivePreviewGeometry->FindLineSet(ExistingSeamsID);
@@ -498,6 +506,8 @@ void UUVEditorSeamTool::ReconstructExistingSeamsVisualization()
 
 void UUVEditorSeamTool::ReconstructLockedPathVisualization()
 {
+	TRACE_CPUPROFILER_EVENT_SCOPE(UVEditorSeamTool_ReconstructLockedPathVisualization);
+
 	using namespace UVEditorSeamToolLocals;
 
 	UPointSetComponent* UnwrapPointSet = UnwrapGeometry->FindPointSet(LockedPointSetID);
@@ -533,6 +543,8 @@ void UUVEditorSeamTool::ReconstructLockedPathVisualization()
 
 void UUVEditorSeamTool::Shutdown(EToolShutdownType ShutdownType)
 {
+	TRACE_CPUPROFILER_EVENT_SCOPE(UVEditorSeamTool_Shutdown);
+
 	UInteractiveTool::Shutdown(ShutdownType);
 
 	EmitChangeAPI = nullptr;
@@ -708,6 +720,8 @@ void UUVEditorSeamTool::ResetPreviewColors()
 
 void UUVEditorSeamTool::ApplyClick()
 {
+	TRACE_CPUPROFILER_EVENT_SCOPE(UVEditorSeamTool_ApplyClick);
+
 	using namespace UVEditorSeamToolLocals;
 
 	ClearHover();
@@ -807,6 +821,8 @@ void UUVEditorSeamTool::ApplyClick()
 
 void UUVEditorSeamTool::ApplySeam(const TArray<int32>& AppliedVidsIn)
 {
+	TRACE_CPUPROFILER_EVENT_SCOPE(UVEditorSeamTool_ApplySeam);
+
 	using namespace UVEditorSeamToolLocals;
 
 	UUVEditorToolMeshInput* ClickedTarget = Targets[ClickedMeshIndex];
@@ -858,6 +874,8 @@ void UUVEditorSeamTool::ApplySeam(const TArray<int32>& AppliedVidsIn)
 
 void UUVEditorSeamTool::ClearLockedPath(bool bEmitChange)
 {
+	TRACE_CPUPROFILER_EVENT_SCOPE(UVEditorSeamTool_ClearLockedPath);
+
 	using namespace UVEditorSeamToolLocals;
 
 	if (bEmitChange)
@@ -887,6 +905,8 @@ void UUVEditorSeamTool::ClearLockedPath(bool bEmitChange)
 
 int32 UUVEditorSeamTool::Get2DHitVertex(const FRay& WorldRayIn, int32* IndexOf2DSpatialOut)
 {
+	TRACE_CPUPROFILER_EVENT_SCOPE(UVEditorSeamTool_Get2DHitVertex);
+
 	if (!ensure(WorldRayIn.Direction.Z != 0))
 	{
 		return IndexConstants::InvalidID;
@@ -942,6 +962,8 @@ int32 UUVEditorSeamTool::Get2DHitVertex(const FRay& WorldRayIn, int32* IndexOf2D
 
 int32 UUVEditorSeamTool::Get3DHitVertex(const FRay& WorldRayIn, int32* IndexOf3DSpatialOut)
 {
+	TRACE_CPUPROFILER_EVENT_SCOPE(UVEditorSeamTool_Get3DHitVertex);
+
 	// When we hit a triangle, gives the closest vert in terms of angle to camera (kind of). Tests
 	// the three verts of the triangle we hit, and the closest vert to the hit location (to make it
 	// easier to hit verts in very thin/degenerate tris).
diff --git a/Engine/Plugins/Experimental/UVEditor/Source/UVEditorTools/Private/UVIslandConformalUnwrapAction.cpp b/Engine/Plugins/Experimental/UVEditor/Source/UVEditorTools/Private/UVIslandConformalUnwrapAction.cpp
--- a/Engine/Plugins/Experimental/UVEditor/Source/UVEditorTools/Private/UVIslandConformalUnwrapAction.cpp
+++ b/Engine/Plugins/Experimental/UVEditor/Source/UVEditorTools/Private/UVIslandConformalUnwrapAction.cpp
@@ -29,12 +29,14 @@ void UUVIslandConformalUnwrapAction::SetWorld(UWorld* WorldIn)
 
 void UUVIslandConformalUnwrapAction::Setup(UInteractiveTool* ParentToolIn)
 {
+	TRACE_CPUPROFILER_EVENT_SCOPE(UVIslandConformalUnwrapAction_Setup);
+
 	UUVToolAction::Setup(ParentToolIn);
 }
 
 void UUVIslandConformalUnwrapAction::Shutdown()
 {
-
+	TRACE_CPUPROFILER_EVENT_SCOPE(UVIslandConformalUnwrapAction_Shutdown);
 }
 
 
@@ -73,6 +75,8 @@ bool UUVIslandConformalUnwrapAction::PreCheckAction()
 }
 
 bool UUVIslandConformalUnwrapAction::GatherIslandTids() {
+	TRACE_CPUPROFILER_EVENT_SCOPE(UVIslandConformalUnwrapAction_GatherIslandTids);
+
 	IslandStartIndices.Reset();
 	ConcatenatedIslandTids.Reset();
 	MaxIslandSize = 0;
@@ -105,6 +109,8 @@ bool UUVIslandConformalUnwrapAction::GatherIslandTids() {
 
 bool UUVIslandConformalUnwrapAction::ApplyAction(UUVToolEmitChangeAPI& EmitChangeAPI)
 {
+	TRACE_CPUPROFILER_EVENT_SCOPE(UVIslandConformalUnwrapAction_ApplyAction);
+
 	FDynamicMesh3& MeshToUnwrap = *(Targets[SelectionTargetIndex]->AppliedCanonical);
 	FDynamicMeshUVOverlay* UVOverlay = MeshToUnwrap.Attributes()->GetUVLayer(Targets[SelectionTargetIndex]->UVLayerIndex);
 	FDynamicMeshUVEditor UVEditor(&MeshToUnwrap, Targets[SelectionTargetIndex]->UVLayerIndex, true);
diff --git a/Engine/Plugins/Experimental/UVEditor/Source/UVEditorTools/Private/UVSeamSewAction.cpp b/Engine/Plugins/Experimental/UVEditor/Source/UVEditorTools/Private/UVSeamSewAction.cpp
--- a/Engine/Plugins/Experimental/UVEditor/Source/UVEditorTools/Private/UVSeamSewAction.cpp
+++ b/Engine/Plugins/Experimental/UVEditor/Source/UVEditorTools/Private/UVSeamSewAction.cpp
@@ -52,6 +52,8 @@ void UUVSeamSewAction::SetWorld(UWorld* WorldIn)
 
 void UUVSeamSewAction::Setup(UInteractiveTool* ParentToolIn)
 {
+	TRACE_CPUPROFILER_EVENT_SCOPE(UVSeamSewAction_Setup);
+	
 	UUVToolAction::Setup(ParentToolIn);
 
 	SewEdgePairingLineSet = NewObject<ULineSetComponent>();
@@ -61,6 +63,8 @@ void UUVSeamSewAction::Setup(UInteractiveTool* ParentToolIn)
 
 void UUVSeamSewAction::Shutdown()
 {
+	TRACE_CPUPROFILER_EVENT_SCOPE(UVSeamSewAction_Shutdown);
+
 	if (UnwrapPreviewGeometryActor)
 	{
 		UnwrapPreviewGeometryActor->Destroy();
@@ -87,6 +91,9 @@ void UUVSeamSewAction::SetSelection(int32 SelectionTargetIndexIn, const UE::Geom
 
 	if (SelectionTargetIndex != -1 && CurrentSelection && CurrentSelection->Type == FDynamicMeshSelection::EType::Edge)
 	{
+		// TODO(Performance) This loop is very slow for large selections
+		TRACE_CPUPROFILER_EVENT_SCOPE(FindSewEdgeOppositePairing_Loop);
+
 		for (int32 Eid : CurrentSelection->SelectedIDs)
 		{
 			FIndex2i SewPairCandidate(Eid, FindSewEdgeOppositePairing(Eid));
@@ -130,7 +137,6 @@ void UUVSeamSewAction::UpdateSewEdgePreviewLines()
 			FVector3d Vert1, Vert2;
 			Vert1 = UnwrapMesh->GetVertex(SewPair.A[0]);
 			Vert2 = UnwrapMesh->GetVertex(SewPair.A[1]);
-			UE_LOG(LogGeometry, Warning, TEXT("Red %d %d."), SewPair.A[0], SewPair.A[1]);
 			SewEdgePairingLineSet->AddLine(
 				MeshTransform.TransformPosition(Vert1),
 				MeshTransform.TransformPosition(Vert2),
@@ -138,7 +144,6 @@ void UUVSeamSewAction::UpdateSewEdgePreviewLines()
 
 			Vert1 = UnwrapMesh->GetVertex(SewPair.B[0]);
 			Vert2 = UnwrapMesh->GetVertex(SewPair.B[1]);
-			UE_LOG(LogGeometry, Warning, TEXT("Green %d %d."), SewPair.B[0], SewPair.B[1]);
 			SewEdgePairingLineSet->AddLine(
 				MeshTransform.TransformPosition(Vert1),
 				MeshTransform.TransformPosition(Vert2),
@@ -179,7 +184,6 @@ int32 UUVSeamSewAction::FindSewEdgeOppositePairing(int32 UnwrapEid) const
 
 	if (!CurrentSelection->Mesh->IsBoundaryEdge(UnwrapEid))
 	{
-		UE_LOG(LogGeometry, Warning, TEXT("Selected edge is unsuitable for sewing - Is not a boundary edge in UV space"));
 		return IndexConstants::InvalidID;
 	}
 	int32 UVLayerIndex = Targets[SelectionTargetIndex]->UVLayerIndex;
@@ -225,6 +229,8 @@ int32 UUVSeamSewAction::FindSewEdgeOppositePairing(int32 UnwrapEid) const
 
 bool UUVSeamSewAction::ApplyAction(UUVToolEmitChangeAPI& EmitChangeAPI)
 {
+	TRACE_CPUPROFILER_EVENT_SCOPE(UVSeamSewAction_ApplyAction);
+
 	FDynamicMesh3& MeshToSew = *(Targets[SelectionTargetIndex]->UnwrapCanonical);
 
 	TArray<int32> SelectedTids;
diff --git a/Engine/Plugins/Experimental/UVEditor/Source/UVEditorToolsEditorOnly/Private/UVEditorParameterizeMeshTool.cpp b/Engine/Plugins/Experimental/UVEditor/Source/UVEditorToolsEditorOnly/Private/UVEditorParameterizeMeshTool.cpp
--- a/Engine/Plugins/Experimental/UVEditor/Source/UVEditorToolsEditorOnly/Private/UVEditorParameterizeMeshTool.cpp
+++ b/Engine/Plugins/Experimental/UVEditor/Source/UVEditorToolsEditorOnly/Private/UVEditorParameterizeMeshTool.cpp
@@ -49,6 +49,8 @@ UInteractiveTool* UUVEditorParameterizeMeshToolBuilder::BuildTool(const FToolBui
 
 void UUVEditorParameterizeMeshTool::Setup()
 {
+	TRACE_CPUPROFILER_EVENT_SCOPE(UVEditorParameterizeMeshTool_Setup);
+
 	check(Targets.Num() >= 1);
 
 	UInteractiveTool::Setup();
@@ -113,6 +115,8 @@ void UUVEditorParameterizeMeshTool::Setup()
 
 void UUVEditorParameterizeMeshTool::OnPropertyModified(UObject* PropertySet, FProperty* Property)
 {
+	TRACE_CPUPROFILER_EVENT_SCOPE(UVEditorParameterizeMeshTool_OnPropertyModified);
+
 	if (PropertySet != MaterialSettings)
 	{
 		for (TObjectPtr<UUVEditorToolMeshInput> Target : Targets)
@@ -131,6 +135,8 @@ void UUVEditorParameterizeMeshTool::OnPropertyModified(UObject* PropertySet, FPr
 
 void UUVEditorParameterizeMeshTool::OnMethodTypeChanged()
 {
+	TRACE_CPUPROFILER_EVENT_SCOPE(UVEditorParameterizeMeshTool_OnMethodTypeChanged);
+	
 	SetToolPropertySourceEnabled(UVAtlasProperties, Settings->Method == EParameterizeMeshUVMethod::UVAtlas);
 	SetToolPropertySourceEnabled(XAtlasProperties, Settings->Method == EParameterizeMeshUVMethod::XAtlas);
 	SetToolPropertySourceEnabled(PatchBuilderProperties, Settings->Method == EParameterizeMeshUVMethod::PatchBuilder);
@@ -144,6 +150,8 @@ void UUVEditorParameterizeMeshTool::OnMethodTypeChanged()
 
 void UUVEditorParameterizeMeshTool::Shutdown(EToolShutdownType ShutdownType)
 {
+	TRACE_CPUPROFILER_EVENT_SCOPE(UVEditorParameterizeMeshTool_Shutdown);
+
 	Settings->SaveProperties(this);
 	MaterialSettings->SaveProperties(this, TEXT("ModelingUVTools"));
 	UVAtlasProperties->SaveProperties(this);
diff --git a/Engine/Plugins/Runtime/MeshModelingToolset/Source/ModelingOperatorsEditorOnly/Private/ParameterizationOps/ParameterizeMeshOp.cpp b/Engine/Plugins/Runtime/MeshModelingToolset/Source/ModelingOperatorsEditorOnly/Private/ParameterizationOps/ParameterizeMeshOp.cpp
--- a/Engine/Plugins/Runtime/MeshModelingToolset/Source/ModelingOperatorsEditorOnly/Private/ParameterizationOps/ParameterizeMeshOp.cpp
+++ b/Engine/Plugins/Runtime/MeshModelingToolset/Source/ModelingOperatorsEditorOnly/Private/ParameterizationOps/ParameterizeMeshOp.cpp
@@ -236,6 +236,8 @@ void FParameterizeMeshOp::CopyNewUVsToMesh(
 
 bool FParameterizeMeshOp::ComputeUVs_UVAtlas(FDynamicMesh3& Mesh,  TFunction<bool(float)>& Interrupter)
 {
+	TRACE_CPUPROFILER_EVENT_SCOPE(ParameterizeMeshOp_ComputeUVs_UVAtlas);
+
 	// the UVAtlas code is unhappy if you feed it a single degenerate triangle
 	bool bNonDegenerate = true;
 	if (Mesh.TriangleCount() == 1)
@@ -311,6 +313,8 @@ bool FParameterizeMeshOp::ComputeUVs_UVAtlas(FDynamicMesh3& Mesh,  TFunction<boo
 
 bool FParameterizeMeshOp::ComputeUVs_XAtlas(FDynamicMesh3& Mesh,  TFunction<bool(float)>& Interrupter)
 {
+	TRACE_CPUPROFILER_EVENT_SCOPE(ParameterizeMeshOp_ComputeUVs_XAtlas);
+
 	// reverse mesh orientation
 	const bool bFixOrientation = true;
 	FDynamicMesh3 FlippedMesh(EMeshComponents::FaceGroups);
@@ -357,6 +361,8 @@ bool FParameterizeMeshOp::ComputeUVs_XAtlas(FDynamicMesh3& Mesh,  TFunction<bool
 
 bool FParameterizeMeshOp::ComputeUVs_PatchBuilder(FDynamicMesh3& InOutMesh, FProgressCancel* ProgressCancel)
 {
+	TRACE_CPUPROFILER_EVENT_SCOPE(ParameterizeMeshOp_ComputeUVs_PatchBuilder);
+	
 	FPatchBasedMeshUVGenerator UVGenerator;
 
 	TUniquePtr<FPolygroupSet> PolygroupConstraint;
@@ -394,6 +400,8 @@ bool FParameterizeMeshOp::ComputeUVs_PatchBuilder(FDynamicMesh3& InOutMesh, FPro
 
 void FParameterizeMeshOp::CalculateResult(FProgressCancel* Progress)
 {
+	TRACE_CPUPROFILER_EVENT_SCOPE(ParameterizeMeshOp_CalculateResult);
+
 	if (!InputMesh.IsValid())
 	{
 		SetResultInfo(FGeometryResult::Failed());
