UV Editor: Add in-viewport selection icons for select tool. Rename select tool to "edit" and change its icon.

#rnx

diff --git a/Engine/Plugins/Experimental/UVEditor/Source/UVEditor/Private/SUVEditor2DViewport.cpp b/Engine/Plugins/Experimental/UVEditor/Source/UVEditor/Private/SUVEditor2DViewport.cpp
--- a/Engine/Plugins/Experimental/UVEditor/Source/UVEditor/Private/SUVEditor2DViewport.cpp
+++ b/Engine/Plugins/Experimental/UVEditor/Source/UVEditor/Private/SUVEditor2DViewport.cpp
@@ -4,9 +4,77 @@
 
 #include "SUVEditor2DViewportToolBar.h"
 #include "UVEditor2DViewportClient.h"
+#include "UVEditorCommands.h"
+#include "UVToolContextObjects.h" // UUVToolViewportButtonsAPI::ESelectionMode
 
 #define LOCTEXT_NAMESPACE "SUVEditor2DViewport"
 
+void SUVEditor2DViewport::BindCommands()
+{
+	SAssetEditorViewport::BindCommands();
+
+	const FUVEditorCommands& CommandInfos = FUVEditorCommands::Get();
+	CommandList->MapAction(
+		CommandInfos.VertexSelection,
+		FExecuteAction::CreateLambda([this]() {
+			StaticCastSharedPtr<FUVEditor2DViewportClient>(Client)->SetSelectionMode(UUVToolViewportButtonsAPI::ESelectionMode::Vertex); 
+		}),
+		FCanExecuteAction::CreateLambda([this]() { 
+			return StaticCastSharedPtr<FUVEditor2DViewportClient>(Client)->AreSelectionButtonsEnabled(); 
+		}),
+		FIsActionChecked::CreateLambda([this]() { 
+			return StaticCastSharedPtr<FUVEditor2DViewportClient>(Client)->GetSelectionMode() == UUVToolViewportButtonsAPI::ESelectionMode::Vertex;
+		}));
+
+	CommandList->MapAction(
+		CommandInfos.EdgeSelection,
+		FExecuteAction::CreateLambda([this]() {
+			StaticCastSharedPtr<FUVEditor2DViewportClient>(Client)->SetSelectionMode(UUVToolViewportButtonsAPI::ESelectionMode::Edge); 
+		}),
+		FCanExecuteAction::CreateLambda([this]() { 
+			return StaticCastSharedPtr<FUVEditor2DViewportClient>(Client)->AreSelectionButtonsEnabled(); 
+		}),
+		FIsActionChecked::CreateLambda([this]() { 
+			return StaticCastSharedPtr<FUVEditor2DViewportClient>(Client)->GetSelectionMode() == UUVToolViewportButtonsAPI::ESelectionMode::Edge;
+		}));
+
+	CommandList->MapAction(
+		CommandInfos.TriangleSelection,
+		FExecuteAction::CreateLambda([this]() {
+			StaticCastSharedPtr<FUVEditor2DViewportClient>(Client)->SetSelectionMode(UUVToolViewportButtonsAPI::ESelectionMode::Triangle); 
+		}),
+		FCanExecuteAction::CreateLambda([this]() { 
+			return StaticCastSharedPtr<FUVEditor2DViewportClient>(Client)->AreSelectionButtonsEnabled(); 
+		}),
+		FIsActionChecked::CreateLambda([this]() { 
+			return StaticCastSharedPtr<FUVEditor2DViewportClient>(Client)->GetSelectionMode() == UUVToolViewportButtonsAPI::ESelectionMode::Triangle;
+		}));
+
+	CommandList->MapAction(
+		CommandInfos.IslandSelection,
+		FExecuteAction::CreateLambda([this]() {
+			StaticCastSharedPtr<FUVEditor2DViewportClient>(Client)->SetSelectionMode(UUVToolViewportButtonsAPI::ESelectionMode::Island); 
+		}),
+		FCanExecuteAction::CreateLambda([this]() { 
+			return StaticCastSharedPtr<FUVEditor2DViewportClient>(Client)->AreSelectionButtonsEnabled(); 
+		}),
+		FIsActionChecked::CreateLambda([this]() { 
+			return StaticCastSharedPtr<FUVEditor2DViewportClient>(Client)->GetSelectionMode() == UUVToolViewportButtonsAPI::ESelectionMode::Island;
+		}));
+
+	CommandList->MapAction(
+		CommandInfos.FullMeshSelection,
+		FExecuteAction::CreateLambda([this]() {
+			StaticCastSharedPtr<FUVEditor2DViewportClient>(Client)->SetSelectionMode(UUVToolViewportButtonsAPI::ESelectionMode::Mesh); 
+		}),
+		FCanExecuteAction::CreateLambda([this]() { 
+			return StaticCastSharedPtr<FUVEditor2DViewportClient>(Client)->AreSelectionButtonsEnabled(); 
+		}),
+		FIsActionChecked::CreateLambda([this]() { 
+			return StaticCastSharedPtr<FUVEditor2DViewportClient>(Client)->GetSelectionMode() == UUVToolViewportButtonsAPI::ESelectionMode::Mesh;
+		}));
+}
+
 void SUVEditor2DViewport::AddOverlayWidget(TSharedRef<SWidget> OverlaidWidget)
 {
 	ViewportOverlay->AddSlot()
diff --git a/Engine/Plugins/Experimental/UVEditor/Source/UVEditor/Private/SUVEditor2DViewportToolBar.cpp b/Engine/Plugins/Experimental/UVEditor/Source/UVEditor/Private/SUVEditor2DViewportToolBar.cpp
--- a/Engine/Plugins/Experimental/UVEditor/Source/UVEditor/Private/SUVEditor2DViewportToolBar.cpp
+++ b/Engine/Plugins/Experimental/UVEditor/Source/UVEditor/Private/SUVEditor2DViewportToolBar.cpp
@@ -6,6 +6,7 @@
 
 #include "EditorViewportCommands.h"
 #include "Framework/MultiBox/MultiBoxBuilder.h"
+#include "UVEditorCommands.h"
 #include "Widgets/Layout/SBorder.h"
 #include "Widgets/SBoxPanel.h"
 
@@ -15,7 +16,6 @@ void SUVEditor2DViewportToolBar::Construct(const FArguments& InArgs)
 	CommandList = InArgs._CommandList;
 
 	const FMargin ToolbarSlotPadding(4.0f, 1.0f);
-	TSharedPtr<SHorizontalBox> MainBoxPtr;
 
 	ChildSlot
 	[
@@ -23,22 +23,76 @@ void SUVEditor2DViewportToolBar::Construct(const FArguments& InArgs)
 		.BorderImage(FAppStyle::Get().GetBrush("EditorViewportToolBar.Background"))
 		.Cursor(EMouseCursor::Default)
 		[
-			SAssignNew( MainBoxPtr, SHorizontalBox )
+			SNew( SHorizontalBox )
+
+			// The first slot is just a spacer so that we get three evenly spaced columns 
+			// and the selection toolbar can go in the center of the center one.
+			+ SHorizontalBox::Slot()
+			.Padding(ToolbarSlotPadding)
+			.HAlign(HAlign_Right)
+
+			+ SHorizontalBox::Slot()
+			.Padding(ToolbarSlotPadding)
+			.HAlign(HAlign_Center)
+			[
+				MakeSelectionToolBar(InArgs._Extenders)
+			]
+
+			+ SHorizontalBox::Slot()
+			.Padding(ToolbarSlotPadding)
+			.HAlign(HAlign_Right)
+			[
+				MakeGizmoToolBar(InArgs._Extenders)
+			]
 		]
 	];
 
-	// Transform toolbar
-	MainBoxPtr->AddSlot()
-		.Padding(ToolbarSlotPadding)
-		.HAlign(HAlign_Right)
-		[
-			MakeToolBar(InArgs._Extenders)
-		];
-
 	SViewportToolBar::Construct(SViewportToolBar::FArguments());
 }
 
-TSharedRef<SWidget> SUVEditor2DViewportToolBar::MakeToolBar(const TSharedPtr<FExtender> InExtenders)
+TSharedRef<SWidget> SUVEditor2DViewportToolBar::MakeSelectionToolBar(const TSharedPtr<FExtender> InExtenders)
+{
+
+	FSlimHorizontalToolBarBuilder ToolbarBuilder(CommandList, FMultiBoxCustomization::None, InExtenders);
+
+	// Use a custom style
+	FName ToolBarStyle = "EditorViewportToolBar";
+	ToolbarBuilder.SetStyle(&FEditorStyle::Get(), ToolBarStyle);
+	ToolbarBuilder.SetLabelVisibility(EVisibility::Collapsed);
+
+	// Transform controls should not be focusable as it fights with the press space to change transform 
+	// mode feature, which we may someday have.
+	ToolbarBuilder.SetIsFocusable(false);
+
+	// Widget controls
+	ToolbarBuilder.BeginSection("SelectionModes");
+	{
+		ToolbarBuilder.BeginBlockGroup();
+
+		ToolbarBuilder.AddToolBarButton(FUVEditorCommands::Get().VertexSelection, NAME_None,
+			TAttribute<FText>(), TAttribute<FText>(), TAttribute<FSlateIcon>(), TEXT("VertexSelection"));
+
+		ToolbarBuilder.AddToolBarButton(FUVEditorCommands::Get().EdgeSelection, NAME_None,
+			TAttribute<FText>(), TAttribute<FText>(), TAttribute<FSlateIcon>(), TEXT("EdgeSelection"));
+
+		ToolbarBuilder.AddToolBarButton(FUVEditorCommands::Get().TriangleSelection, NAME_None,
+			TAttribute<FText>(), TAttribute<FText>(), TAttribute<FSlateIcon>(), TEXT("TriangleSelection"));
+
+		ToolbarBuilder.AddToolBarButton(FUVEditorCommands::Get().IslandSelection, NAME_None,
+			TAttribute<FText>(), TAttribute<FText>(), TAttribute<FSlateIcon>(), TEXT("IslandSelection"));
+
+		ToolbarBuilder.AddToolBarButton(FUVEditorCommands::Get().FullMeshSelection, NAME_None,
+			TAttribute<FText>(), TAttribute<FText>(), TAttribute<FSlateIcon>(), TEXT("FullMeshSelection"));
+
+		ToolbarBuilder.EndBlockGroup();
+	}
+
+	ToolbarBuilder.EndSection();
+
+	return ToolbarBuilder.MakeWidget();
+}
+
+TSharedRef<SWidget> SUVEditor2DViewportToolBar::MakeGizmoToolBar(const TSharedPtr<FExtender> InExtenders)
 {
 	// The following is modeled after portions of STransformViewportToolBar, which gets 
 	// used in SCommonEditorViewportToolbarBase.
diff --git a/Engine/Plugins/Experimental/UVEditor/Source/UVEditor/Private/UVEditor2DViewportClient.cpp b/Engine/Plugins/Experimental/UVEditor/Source/UVEditor/Private/UVEditor2DViewportClient.cpp
--- a/Engine/Plugins/Experimental/UVEditor/Source/UVEditor/Private/UVEditor2DViewportClient.cpp
+++ b/Engine/Plugins/Experimental/UVEditor/Source/UVEditor/Private/UVEditor2DViewportClient.cpp
@@ -201,3 +201,26 @@ UE::Widget::EWidgetMode FUVEditor2DViewportClient::GetWidgetMode() const
 		break;
 	}
 }
+
+bool FUVEditor2DViewportClient::AreSelectionButtonsEnabled() const
+{
+	return ViewportButtonsAPI && ViewportButtonsAPI->AreSelectionButtonsEnabled();
+}
+
+void FUVEditor2DViewportClient::SetSelectionMode(UUVToolViewportButtonsAPI::ESelectionMode NewMode)
+{
+	if (ViewportButtonsAPI)
+	{
+		ViewportButtonsAPI->SetSelectionMode(NewMode);
+	}
+}
+
+UUVToolViewportButtonsAPI::ESelectionMode FUVEditor2DViewportClient::GetSelectionMode() const
+{
+	if (!AreSelectionButtonsEnabled())
+	{
+		return UUVToolViewportButtonsAPI::ESelectionMode::None;
+	}
+
+	return ViewportButtonsAPI->GetSelectionMode();
+}
\ No newline at end of file
diff --git a/Engine/Plugins/Experimental/UVEditor/Source/UVEditor/Private/UVEditorCommands.cpp b/Engine/Plugins/Experimental/UVEditor/Source/UVEditor/Private/UVEditorCommands.cpp
--- a/Engine/Plugins/Experimental/UVEditor/Source/UVEditor/Private/UVEditorCommands.cpp
+++ b/Engine/Plugins/Experimental/UVEditor/Source/UVEditor/Private/UVEditorCommands.cpp
@@ -24,7 +24,7 @@ void FUVEditorCommands::RegisterCommands()
 	UI_COMMAND(ApplyChanges, "Apply Changes", "Apply changes without yet saving the assets", EUserInterfaceActionType::Button, FInputChord());
 
 	// These get linked to various tool buttons.
-	UI_COMMAND(BeginSelectTool, "Select", "Switch to select tool", EUserInterfaceActionType::ToggleButton, FInputChord());
+	UI_COMMAND(BeginSelectTool, "Edit", "Switch to selection-based edit tool", EUserInterfaceActionType::ToggleButton, FInputChord());
 	UI_COMMAND(BeginLayoutTool, "Layout", "Switch to layout tool", EUserInterfaceActionType::ToggleButton, FInputChord());
 	UI_COMMAND(BeginParameterizeMeshTool, "Auto UV", "Switch to Auto UV tool", EUserInterfaceActionType::ToggleButton, FInputChord());
 	UI_COMMAND(BeginChannelEditTool, "Channel Edit", "Switch to Channel Edit tool", EUserInterfaceActionType::ToggleButton, FInputChord());
@@ -35,6 +35,13 @@ void FUVEditorCommands::RegisterCommands()
 	UI_COMMAND(AcceptOrCompleteActiveTool, "Accept", "Accept the active tool", EUserInterfaceActionType::Button, FInputChord(EKeys::Enter));
 	UI_COMMAND(CancelOrCompleteActiveTool, "Cancel", "Cancel the active tool", EUserInterfaceActionType::Button, FInputChord(EKeys::Escape));
 
+	// These get used in viewport buttons
+	UI_COMMAND(VertexSelection, "Vertex Selection", "Select vertex on click", EUserInterfaceActionType::ToggleButton, FInputChord());
+	UI_COMMAND(EdgeSelection, "Edge Selection", "Select edge on click", EUserInterfaceActionType::ToggleButton, FInputChord());
+	UI_COMMAND(TriangleSelection, "Triangle Selection", "Select triangle on click", EUserInterfaceActionType::ToggleButton, FInputChord());
+	UI_COMMAND(IslandSelection, "Island Selection", "Select connected island on click", EUserInterfaceActionType::ToggleButton, FInputChord());
+	UI_COMMAND(FullMeshSelection, "Full Mesh Selection", "Select whole mesh on click", EUserInterfaceActionType::ToggleButton, FInputChord());
+
 	UI_COMMAND(EnableOrbitCamera, "Orbit", "Enable Orbit Camera", EUserInterfaceActionType::ToggleButton, FInputChord());
 	UI_COMMAND(EnableFlyCamera, "Fly", "Enable Fly Camera", EUserInterfaceActionType::ToggleButton, FInputChord());
 }
diff --git a/Engine/Plugins/Experimental/UVEditor/Source/UVEditor/Private/UVEditorStyle.cpp b/Engine/Plugins/Experimental/UVEditor/Source/UVEditor/Private/UVEditorStyle.cpp
--- a/Engine/Plugins/Experimental/UVEditor/Source/UVEditor/Private/UVEditorStyle.cpp
+++ b/Engine/Plugins/Experimental/UVEditor/Source/UVEditor/Private/UVEditorStyle.cpp
@@ -14,26 +14,35 @@ FUVEditorStyle::FUVEditorStyle()
 {
 	// Used FFractureEditorStyle as a model
 
-	const FVector2D IconSize(20.0f, 20.0f);
+	const FVector2D IconSize(16.0f, 16.0f);
+	const FVector2D ToolbarIconSize(20.0f, 20.0f);
 
 	SetContentRoot(FPaths::EnginePluginsDir() / TEXT("Experimental/UVEditor/Content/Icons"));
 	SetCoreContentRoot(FPaths::EngineContentDir() / TEXT("Editor/Slate") );
 
-	Set("UVEditor.BeginSelectTool",             new IMAGE_BRUSH_SVG("MeshSelect", IconSize));
-	Set("UVEditor.BeginLayoutTool",             new IMAGE_BRUSH_SVG("UVLayout", IconSize));
-	Set("UVEditor.BeginParameterizeMeshTool",   new IMAGE_BRUSH_SVG("AutoUnwrap", IconSize));
-	Set("UVEditor.BeginChannelEditTool",        new IMAGE_BRUSH_SVG("AttributeEditor", IconSize));
-	Set("UVEditor.BeginSeamTool",               new IMAGE_BRUSH_SVG("ModelingUVSeamEdit", IconSize));
-	Set("UVEditor.BeginRecomputeUVsTool",       new IMAGE_BRUSH_SVG("GroupUnwrap", IconSize));
+	Set("UVEditor.OpenUVEditor",				new IMAGE_BRUSH_SVG("UVEditor", IconSize));
+
+	Set("UVEditor.BeginSelectTool",				new FSlateVectorImageBrush(FPaths::EngineContentDir() / TEXT("Slate/Starship/Common/edit.svg"), ToolbarIconSize));
+	Set("UVEditor.BeginLayoutTool",             new IMAGE_BRUSH_SVG("UVLayout", ToolbarIconSize));
+	Set("UVEditor.BeginParameterizeMeshTool",   new IMAGE_BRUSH_SVG("AutoUnwrap", ToolbarIconSize));
+	Set("UVEditor.BeginChannelEditTool",        new IMAGE_BRUSH_SVG("AttributeEditor", ToolbarIconSize));
+	Set("UVEditor.BeginSeamTool",               new IMAGE_BRUSH_SVG("ModelingUVSeamEdit", ToolbarIconSize));
+	Set("UVEditor.BeginRecomputeUVsTool",       new IMAGE_BRUSH_SVG("GroupUnwrap", ToolbarIconSize));
 
 	// Top toolbar icons
-	Set("UVEditor.ApplyChanges", new CORE_IMAGE_BRUSH_SVG("Starship/Common/Apply", IconSize));
-	Set("UVEditor.ChannelSettings", new CORE_IMAGE_BRUSH_SVG("Starship/Common/SetDrawUVs", IconSize));
-	Set("UVEditor.BackgroundSettings", new CORE_IMAGE_BRUSH_SVG("Starship/Common/Sprite", IconSize));
+	Set("UVEditor.ApplyChanges", new CORE_IMAGE_BRUSH_SVG("Starship/Common/Apply", ToolbarIconSize));
+	Set("UVEditor.ChannelSettings", new CORE_IMAGE_BRUSH_SVG("Starship/Common/SetDrawUVs", ToolbarIconSize));
+	Set("UVEditor.BackgroundSettings", new CORE_IMAGE_BRUSH_SVG("Starship/Common/Sprite", ToolbarIconSize));
 
 	// Viewport icons
-	Set("UVEditor.OrbitCamera", new CORE_IMAGE_BRUSH_SVG("Starship/EditorViewport/rotate", IconSize));
-	Set("UVEditor.FlyCamera", new CORE_IMAGE_BRUSH_SVG("Starship/EditorViewport/camera", IconSize));
+	Set("UVEditor.OrbitCamera", new CORE_IMAGE_BRUSH_SVG("Starship/EditorViewport/rotate", ToolbarIconSize));
+	Set("UVEditor.FlyCamera", new CORE_IMAGE_BRUSH_SVG("Starship/EditorViewport/camera", ToolbarIconSize));
+
+	Set("UVEditor.VertexSelection", new IMAGE_BRUSH_SVG("SelectionVertices", ToolbarIconSize));
+	Set("UVEditor.EdgeSelection", new IMAGE_BRUSH_SVG("SelectionLine", ToolbarIconSize));
+	Set("UVEditor.TriangleSelection", new IMAGE_BRUSH_SVG("SelectionTriangle", ToolbarIconSize));
+	Set("UVEditor.IslandSelection", new IMAGE_BRUSH_SVG("SelectionIslands", ToolbarIconSize));
+	Set("UVEditor.FullMeshSelection", new IMAGE_BRUSH_SVG("SelectionMulti", ToolbarIconSize));
 
 	FSlateStyleRegistry::RegisterSlateStyle(*this);
 }
diff --git a/Engine/Plugins/Experimental/UVEditor/Source/UVEditor/Public/SUVEditor2DViewport.h b/Engine/Plugins/Experimental/UVEditor/Source/UVEditor/Public/SUVEditor2DViewport.h
--- a/Engine/Plugins/Experimental/UVEditor/Source/UVEditor/Public/SUVEditor2DViewport.h
+++ b/Engine/Plugins/Experimental/UVEditor/Source/UVEditor/Public/SUVEditor2DViewport.h
@@ -18,6 +18,7 @@ public:
 	virtual void RemoveOverlayWidget(TSharedRef<SWidget> OverlaidWidget);
 
 	// SEditorViewport
+	virtual void BindCommands() override;
 	virtual TSharedPtr<SWidget> MakeViewportToolbar() override;
 	bool IsWidgetModeActive(UE::Widget::EWidgetMode Mode) const override;
 };
diff --git a/Engine/Plugins/Experimental/UVEditor/Source/UVEditor/Public/SUVEditor2DViewportToolBar.h b/Engine/Plugins/Experimental/UVEditor/Source/UVEditor/Public/SUVEditor2DViewportToolBar.h
--- a/Engine/Plugins/Experimental/UVEditor/Source/UVEditor/Public/SUVEditor2DViewportToolBar.h
+++ b/Engine/Plugins/Experimental/UVEditor/Source/UVEditor/Public/SUVEditor2DViewportToolBar.h
@@ -24,7 +24,8 @@ public:
 	void Construct(const FArguments& InArgs);
 
 private:
-	TSharedRef<SWidget> MakeToolBar(const TSharedPtr<FExtender> InExtenders);
+	TSharedRef<SWidget> MakeSelectionToolBar(const TSharedPtr<FExtender> InExtenders);
+	TSharedRef<SWidget> MakeGizmoToolBar(const TSharedPtr<FExtender> InExtenders);
 
 	/** Command list */
 	TSharedPtr<FUICommandList> CommandList;
diff --git a/Engine/Plugins/Experimental/UVEditor/Source/UVEditor/Public/UVEditor2DViewportClient.h b/Engine/Plugins/Experimental/UVEditor/Source/UVEditor/Public/UVEditor2DViewportClient.h
--- a/Engine/Plugins/Experimental/UVEditor/Source/UVEditor/Public/UVEditor2DViewportClient.h
+++ b/Engine/Plugins/Experimental/UVEditor/Source/UVEditor/Public/UVEditor2DViewportClient.h
@@ -7,9 +7,7 @@
 #include "EditorViewportClient.h"
 #include "InputBehaviorSet.h"
 #include "UVEditor2DViewportBehaviorTargets.h" // FUVEditor2DScrollBehaviorTarget, FUVEditor2DMouseWheelZoomBehaviorTarget
-#include "UVToolContextObjects.h" // UUVToolViewportButtonsAPI::EGizmoMode
-
-class UUVToolViewportButtonsAPI;
+#include "UVToolContextObjects.h" // UUVToolViewportButtonsAPI::ESelectionMode
 
 /**
  * Client used to display a 2D view of the UV's, implemented by using a perspective viewport with a locked
@@ -23,6 +21,10 @@ public:
 
 	virtual ~FUVEditor2DViewportClient() {}
 
+	bool AreSelectionButtonsEnabled() const;
+	void SetSelectionMode(UUVToolViewportButtonsAPI::ESelectionMode NewMode);
+	UUVToolViewportButtonsAPI::ESelectionMode GetSelectionMode() const;
+
 	bool AreWidgetButtonsEnabled() const;
 
 	// FEditorViewportClient
diff --git a/Engine/Plugins/Experimental/UVEditor/Source/UVEditor/Public/UVEditorCommands.h b/Engine/Plugins/Experimental/UVEditor/Source/UVEditor/Public/UVEditorCommands.h
--- a/Engine/Plugins/Experimental/UVEditor/Source/UVEditor/Public/UVEditorCommands.h
+++ b/Engine/Plugins/Experimental/UVEditor/Source/UVEditor/Public/UVEditorCommands.h
@@ -28,6 +28,12 @@ public:
 	TSharedPtr<FUICommandInfo> AcceptOrCompleteActiveTool;
 	TSharedPtr<FUICommandInfo> CancelOrCompleteActiveTool;
 
+	TSharedPtr<FUICommandInfo> VertexSelection;
+	TSharedPtr<FUICommandInfo> EdgeSelection;
+	TSharedPtr<FUICommandInfo> TriangleSelection;
+	TSharedPtr<FUICommandInfo> IslandSelection;
+	TSharedPtr<FUICommandInfo> FullMeshSelection;
+
 	TSharedPtr<FUICommandInfo> EnableOrbitCamera;
 	TSharedPtr<FUICommandInfo> EnableFlyCamera;
 };
\ No newline at end of file
diff --git a/Engine/Plugins/Experimental/UVEditor/Source/UVEditorTools/Private/UVSelectTool.cpp b/Engine/Plugins/Experimental/UVEditor/Source/UVEditorTools/Private/UVSelectTool.cpp
--- a/Engine/Plugins/Experimental/UVEditor/Source/UVEditorTools/Private/UVSelectTool.cpp
+++ b/Engine/Plugins/Experimental/UVEditor/Source/UVEditorTools/Private/UVSelectTool.cpp
@@ -306,10 +306,6 @@ void UUVSelectTool::Setup()
 	
 	SetToolDisplayName(LOCTEXT("ToolName", "UV Select Tool"));
 
-	Settings = NewObject<UUVSelectToolProperties>(this);
-	Settings->RestoreProperties(this);
-	AddToolPropertySource(Settings);
-
 	UContextObjectStore* ContextStore = GetToolManager()->GetContextObjectStore();
 	EmitChangeAPI = ContextStore->FindContext<UUVToolEmitChangeAPI>();
 	ViewportButtonsAPI = ContextStore->FindContext<UUVToolViewportButtonsAPI>();
@@ -318,6 +314,11 @@ void UUVSelectTool::Setup()
 		[this](UUVToolViewportButtonsAPI::EGizmoMode NewGizmoMode) {
 			UpdateGizmo();
 		});
+	ViewportButtonsAPI->SetSelectionButtonsEnabled(true);
+	ViewportButtonsAPI->OnSelectionModeChange.AddWeakLambda(this,
+		[this](UUVToolViewportButtonsAPI::ESelectionMode NewMode) {
+			UpdateSelectionMode();
+		});
 
 	ToolActions = NewObject<USelectToolActionPropertySet>(this);
 	ToolActions->Initialize(this);
@@ -359,7 +360,7 @@ void UUVSelectTool::Setup()
 			LOCTEXT("SelectionChangeMessage", "Selection Change"));
 	};
 
-	ConfigureSelectionModeFromControls();
+	UpdateSelectionMode();
 
 	// Retrieve cached AABB tree storage, or else set it up
 	UUVToolAABBTreeStorage* TreeStore = ContextStore->FindContext<UUVToolAABBTreeStorage>();
@@ -467,8 +468,6 @@ void UUVSelectTool::Shutdown(EToolShutdownType ShutdownType)
 		Target->OnCanonicalModified.RemoveAll(this);
 	}
 
-	Settings->SaveProperties(this);
-
 	SelectionMechanic->Shutdown();
 
 	if (LivePreviewGeometryActor)
@@ -491,7 +490,9 @@ void UUVSelectTool::Shutdown(EToolShutdownType ShutdownType)
 	GetToolManager()->GetPairedGizmoManager()->DestroyAllGizmosByOwner(this);
 
 	ViewportButtonsAPI->OnGizmoModeChange.RemoveAll(this);
+	ViewportButtonsAPI->OnSelectionModeChange.RemoveAll(this);
 	ViewportButtonsAPI->SetGizmoButtonsEnabled(false);
+	ViewportButtonsAPI->SetSelectionButtonsEnabled(false);
 
 	ViewportButtonsAPI = nullptr;
 	EmitChangeAPI = nullptr;
@@ -511,7 +512,6 @@ void UUVSelectTool::SetGizmoTransform(const FTransform& NewTransform)
 
 void UUVSelectTool::OnPropertyModified(UObject* PropertySet, FProperty* Property)
 {
-	ConfigureSelectionModeFromControls();
 }
 
 void UUVSelectTool::UpdateGizmo()
@@ -530,28 +530,30 @@ void UUVSelectTool::UpdateGizmo()
 		&& !SelectionMechanic->GetCurrentSelection().IsEmpty());
 }
 
-void UUVSelectTool::ConfigureSelectionModeFromControls()
+void UUVSelectTool::UpdateSelectionMode()
 {
 	EMeshSelectionMechanicMode TargetMode;
-	switch (Settings->SelectionMode)
+	switch (ViewportButtonsAPI->GetSelectionMode())
 	{
-	case EUVSelectToolSelectionMode::Island:
-		TargetMode = EMeshSelectionMechanicMode::Component;
-		break;
-	case EUVSelectToolSelectionMode::Edge:
-		TargetMode = EMeshSelectionMechanicMode::Edge;
-		break;
-	case EUVSelectToolSelectionMode::Vertex:
+	case UUVToolViewportButtonsAPI::ESelectionMode::Vertex:
 		TargetMode = EMeshSelectionMechanicMode::Vertex;
+		break;	
+	case UUVToolViewportButtonsAPI::ESelectionMode::Edge:
+		TargetMode = EMeshSelectionMechanicMode::Edge;
 		break;
-	case EUVSelectToolSelectionMode::Triangle:
+	case UUVToolViewportButtonsAPI::ESelectionMode::Triangle:
 		TargetMode = EMeshSelectionMechanicMode::Triangle;
 		break;
-	case EUVSelectToolSelectionMode::Mesh:
+	case UUVToolViewportButtonsAPI::ESelectionMode::Island:
+		TargetMode = EMeshSelectionMechanicMode::Component;
+		break;
+	case UUVToolViewportButtonsAPI::ESelectionMode::Mesh:
 		TargetMode = EMeshSelectionMechanicMode::Mesh;
 		break;
 	default:
+		// We shouldn't ever get "none" as the selection mode...
 		ensure(false);
+		TargetMode = EMeshSelectionMechanicMode::Vertex;
 		break;
 	}
 	SelectionMechanic->ChangeSelectionMode(TargetMode);
@@ -767,17 +769,9 @@ void UUVSelectTool::GizmoTransformEnded(UTransformProxy* Proxy)
 	// One final attempt to apply transforms if OnTick hasn't happened yet
 	ApplyGizmoTransform();
 
-	if (Settings->bUpdatePreviewDuringDrag)
-	{
-		// Both previews must already be updated, so only need to update canonical
-		Targets[SelectionTargetIndex]->UpdateCanonicalFromPreviews(&MovingVids, 
-			UUVEditorToolMeshInput::NONE_CHANGED_ARG);
-	}
-	else
-	{
-		Targets[SelectionTargetIndex]->UpdateAllFromUnwrapPreview(&MovingVids, 
-			UUVEditorToolMeshInput::NONE_CHANGED_ARG, &SelectedTids);
-	}
+	// Both previews must already be updated, so only need to update canonical
+	Targets[SelectionTargetIndex]->UpdateCanonicalFromPreviews(&MovingVids, 
+		UUVEditorToolMeshInput::NONE_CHANGED_ARG);
 
 	const FText TransactionName(LOCTEXT("DragCompleteTransactionName", "Move Items"));
 	EmitChangeAPI->EmitToolIndependentChange(ChangeRouter, MakeUnique<UVSelectToolLocals::FGizmoMeshChange>(
@@ -815,11 +809,8 @@ void UUVSelectTool::ApplyGizmoTransform()
 
 		SelectionMechanic->SetDrawnElementsTransform((FTransform)TransformToApply);
 
-		if (Settings->bUpdatePreviewDuringDrag)
-		{
-			Targets[SelectionTargetIndex]->UpdateAppliedPreviewFromUnwrapPreview(&MovingVids, 
-				UUVEditorToolMeshInput::NONE_CHANGED_ARG, &SelectedTids);
-		}
+		Targets[SelectionTargetIndex]->UpdateAppliedPreviewFromUnwrapPreview(&MovingVids, 
+			UUVEditorToolMeshInput::NONE_CHANGED_ARG, &SelectedTids);
 
 		bGizmoTransformNeedsApplication = false;
 		SewAction->UpdateVisualizations();
diff --git a/Engine/Plugins/Experimental/UVEditor/Source/UVEditorTools/Public/UVSelectTool.h b/Engine/Plugins/Experimental/UVEditor/Source/UVEditorTools/Public/UVSelectTool.h
--- a/Engine/Plugins/Experimental/UVEditor/Source/UVEditorTools/Public/UVSelectTool.h
+++ b/Engine/Plugins/Experimental/UVEditor/Source/UVEditorTools/Public/UVSelectTool.h
@@ -50,33 +50,6 @@ protected:
 	virtual const FToolTargetTypeRequirements& GetTargetRequirements() const override;
 };
 
-UENUM()
-enum class EUVSelectToolSelectionMode : uint8
-{
-	Vertex,
-	Edge,
-	Triangle,
-	Island,
-	Mesh
-};
-
-UCLASS()
-class UVEDITORTOOLS_API UUVSelectToolProperties : public UInteractiveToolPropertySet
-{
-	GENERATED_BODY()
-
-public:
-
-	UPROPERTY(EditAnywhere, Category = Options)
-	EUVSelectToolSelectionMode SelectionMode = EUVSelectToolSelectionMode::Island;
-
-	//~ TODO: Make this only visible in transform mode
-	UPROPERTY(EditAnywhere, Category = Options, AdvancedDisplay)
-	bool bUpdatePreviewDuringDrag = true;
-};
-
-
-
 UENUM()
 enum class ESelectToolAction
 {
@@ -166,18 +139,13 @@ protected:
 	virtual void UpdateGizmo();
 	virtual void UpdateLivePreviewLines();
 
-	void SetToolActionsVisiblity(bool isVisible);
-
-	void ConfigureSelectionModeFromControls();	
+	void UpdateSelectionMode();	
 
 	UWorld* TargetWorld;
 
 	UPROPERTY()
 	TArray<TObjectPtr<UUVEditorToolMeshInput>> Targets;
 
-	UPROPERTY()
-	TObjectPtr<UUVSelectToolProperties> Settings = nullptr;
-
 	UPROPERTY()
 	TObjectPtr<USelectToolActionPropertySet> ToolActions = nullptr;
 
diff --git a/Engine/Plugins/Experimental/UVEditor/Source/UVEditorTools/Public/UVToolContextObjects.h b/Engine/Plugins/Experimental/UVEditor/Source/UVEditorTools/Public/UVToolContextObjects.h
--- a/Engine/Plugins/Experimental/UVEditor/Source/UVEditorTools/Public/UVToolContextObjects.h
+++ b/Engine/Plugins/Experimental/UVEditor/Source/UVEditorTools/Public/UVToolContextObjects.h
@@ -117,7 +117,7 @@ protected:
 };
 
 /**
- * Allows tools to interact with buttons in the viewport (currently just gizmo controls)
+ * Allows tools to interact with buttons in the viewport
  */
 UCLASS()
 class UVEDITORTOOLS_API UUVToolViewportButtonsAPI : public UUVToolContextObject
@@ -131,14 +131,25 @@ public:
 		Transform
 	};
 
+	enum class ESelectionMode
+	{
+		None,
+
+		Vertex,
+		Edge,
+		Triangle,
+		Island,
+		Mesh,
+	};
+
 	void SetGizmoButtonsEnabled(bool bOn)
 	{
-		bButtonsEnabled = bOn;
+		bGizmoButtonsEnabled = bOn;
 	}
 
 	bool AreGizmoButtonsEnabled()
 	{
-		return bButtonsEnabled;
+		return bGizmoButtonsEnabled;
 	}
 
 	void SetGizmoMode(EGizmoMode ModeIn, bool bBroadcast = true)
@@ -152,13 +163,40 @@ public:
 		return GizmoMode;
 	}
 
+
+	void SetSelectionButtonsEnabled(bool bOn)
+	{
+		bSelectionButtonsEnabled = bOn;
+	}
+
+	bool AreSelectionButtonsEnabled()
+	{
+		return bSelectionButtonsEnabled;
+	}
+
+	void SetSelectionMode(ESelectionMode ModeIn, bool bBroadcast = true)
+	{
+		SelectionMode = ModeIn;
+		OnSelectionModeChange.Broadcast(SelectionMode);
+	}
+
+	ESelectionMode GetSelectionMode()
+	{
+		return SelectionMode;
+	}
+
 	DECLARE_MULTICAST_DELEGATE_OneParam(FOnGizmoModeChange, EGizmoMode NewGizmoMode);
 	FOnGizmoModeChange OnGizmoModeChange;
 
+	DECLARE_MULTICAST_DELEGATE_OneParam(FOnSelectionModeChange, ESelectionMode NewSelectionMode);
+	FOnSelectionModeChange OnSelectionModeChange;
+
 protected:
 
-	bool bButtonsEnabled = false;
+	bool bGizmoButtonsEnabled = false;
 	EGizmoMode GizmoMode = EGizmoMode::Select;
+	bool bSelectionButtonsEnabled = false;
+	ESelectionMode SelectionMode = ESelectionMode::Island;
 };
 
 
