UVEditor: Improve performance of opposite edge determination for sewing action, by adjusting algorithm used to locate edges and removing console log instructions.

#rnx

diff --git a/Engine/Plugins/Experimental/UVEditor/Source/UVEditorTools/Private/UVSeamSewAction.cpp b/Engine/Plugins/Experimental/UVEditor/Source/UVEditorTools/Private/UVSeamSewAction.cpp
--- a/Engine/Plugins/Experimental/UVEditor/Source/UVEditorTools/Private/UVSeamSewAction.cpp
+++ b/Engine/Plugins/Experimental/UVEditor/Source/UVEditorTools/Private/UVSeamSewAction.cpp
@@ -188,42 +188,46 @@ int32 UUVSeamSewAction::FindSewEdgeOppositePairing(int32 UnwrapEid) const
 	}
 	int32 UVLayerIndex = Targets[SelectionTargetIndex]->UVLayerIndex;
 	FDynamicMeshUVOverlay* UVOverlay = Targets[SelectionTargetIndex]->AppliedCanonical.Get()->Attributes()->GetUVLayer(UVLayerIndex);
-	FIndex2i EdgeVids = CurrentSelection->Mesh->GetEdgeV(UnwrapEid);
-	int32 ParentVid_0 = UVOverlay->GetParentVertex(EdgeVids[0]);
-	int32 ParentVid_1 = UVOverlay->GetParentVertex(EdgeVids[1]);
-
-	TArray<int32> Vid0Elements;
-	TArray<int32> Vid1Elements;
-	UVOverlay->GetVertexElements(ParentVid_0, Vid0Elements);
-	UVOverlay->GetVertexElements(ParentVid_1, Vid1Elements);
-
-	// Now that we know our potential elements, locate one single edge that is a viable pairing, or return InvalidID if there are none or more.
-	int32 OppositeEid = IndexConstants::InvalidID;
-	for (int32 Vid0 : Vid0Elements)
+	UE::Geometry::FDynamicMesh3::FEdge UnwrapEdge = CurrentSelection->Mesh->GetEdge(UnwrapEid);
+	checkSlow(UnwrapEdge.Tri[1] == -1); // As a boundary edge, the second triangle should always be invalid.
+	int32 ParentVid_0 = UVOverlay->GetParentVertex(UnwrapEdge.Vert[0]);
+	int32 ParentVid_1 = UVOverlay->GetParentVertex(UnwrapEdge.Vert[1]);
+
+	int32 ParentEid = Targets[SelectionTargetIndex]->AppliedCanonical->FindEdgeFromTri(ParentVid_0, ParentVid_1, UnwrapEdge.Tri[0]);
+	FIndex2i AppliedEdgeTids = Targets[SelectionTargetIndex]->AppliedCanonical->GetEdgeT(ParentEid);
+
+	int32 OppositeTid = UnwrapEdge.Tri[0] == AppliedEdgeTids[0] ? AppliedEdgeTids[1] : AppliedEdgeTids[0];
+	if (OppositeTid == -1) 
 	{
-		for (int32 Vid1 : Vid1Elements)
-		{
-			int32 EidCandidate = CurrentSelection->Mesh->FindEdge(Vid0, Vid1);
-			if (EidCandidate == UnwrapEid || EidCandidate == IndexConstants::InvalidID)
-			{
-				continue;
-			}
-			if (OppositeEid == IndexConstants::InvalidID)
-			{
-				OppositeEid = EidCandidate;
-				int32 ParentCandidateVid_0 = UVOverlay->GetParentVertex(Vid0);
-				int32 ParentCandidateVid_1 = UVOverlay->GetParentVertex(Vid1);
-				check(ParentCandidateVid_0 == ParentVid_0 && ParentCandidateVid_1 == ParentVid_1)
-			}
-			else
-			{
-				UE_LOG(LogGeometry, Warning, TEXT("Selected edge is unsuitable for sewing - there are non-manifold candidate edges."));
-				return IndexConstants::InvalidID;
-			}
+		// This could happen if a boundary edge in the unwrap mesh is also a boundary edge in the applied, i.e. the applied mesh isn't closed.
+		return IndexConstants::InvalidID;
+	}
+	if (!CurrentSelection->Mesh->IsTriangle(OppositeTid))
+	{
+		// Return invalid if the opposite triangle isn't set in the unwtap mesh, i.e. the overlay has incomplete UVs
+		return IndexConstants::InvalidID;
+	}
+
+	FIndex3i UnwrapOppositeEids = CurrentSelection->Mesh->GetTriEdges(OppositeTid);
+	for (int32 i = 0; i < 3; ++i)
+	{
+		FIndex2i UnwrapOppositeVids = CurrentSelection->Mesh->GetEdgeV(UnwrapOppositeEids[i]);
+		if (!ensure(UVOverlay->IsElement(UnwrapOppositeVids[0]) && UVOverlay->IsElement(UnwrapOppositeVids[1]))) {			
+			continue; // Skip in case any elements aren't properly set in the overlay.
+			          //We probably shouldn't reach this if the triangle check above passed.
+		}
+		int32 OppositeParentVid_0 = UVOverlay->GetParentVertex(UnwrapOppositeVids[0]);
+		int32 OppositeParentVid_1 = UVOverlay->GetParentVertex(UnwrapOppositeVids[1]);
+
+		if ((OppositeParentVid_0 == ParentVid_0 && OppositeParentVid_1 == ParentVid_1) ||
+			(OppositeParentVid_0 == ParentVid_1 && OppositeParentVid_1 == ParentVid_0)) {
+			return UnwrapOppositeEids[i];
 		}
 	}
 
-	return OppositeEid;
+	// If we can't find anything... return an invalid result.
+	return IndexConstants::InvalidID;
+
 }
 
 
