UVEditor: Improve handling of tool bar palette tab in the UVEditor to use current best practices of UILayers and integrated palette construction logic from FModeToolkit. Also fixes the graphical issue where toolbar buttons change position depending on how the toolbar tab well is shrunk or expanded.

#rnx

diff --git a/Engine/Plugins/Experimental/UVEditor/Source/UVEditor/Private/UVEditorModeToolkit.cpp b/Engine/Plugins/Experimental/UVEditor/Source/UVEditor/Private/UVEditorModeToolkit.cpp
--- a/Engine/Plugins/Experimental/UVEditor/Source/UVEditor/Private/UVEditorModeToolkit.cpp
+++ b/Engine/Plugins/Experimental/UVEditor/Source/UVEditor/Private/UVEditorModeToolkit.cpp
@@ -17,6 +17,9 @@
 #include "Widgets/Layout/SBorder.h"
 #include "Widgets/Text/STextBlock.h"
 #include "EdModeInteractiveToolsContext.h"
+#include "EditorModeManager.h"
+#include "UVEditorModeUILayer.h"
+#include "AssetEditorModeManager.h"
 
 #define LOCTEXT_NAMESPACE "FUVEditorModeToolkit"
 
@@ -129,22 +132,6 @@ void FUVEditorModeToolkit::Init(const TSharedPtr<IToolkitHost>& InitToolkitHost,
 	// of tool start transactions. See FUVEditorModeToolkit::OnToolStarted for how we issue the transactions.
 	InteractiveToolsContext->ToolManager->ConfigureChangeTrackingMode(EToolChangeTrackingMode::NoChangeTracking);
 
-	// Build the tool palette
-	const FUVEditorCommands& CommandInfos = FUVEditorCommands::Get();
-	TSharedRef<FUICommandList> CommandList = GetToolkitCommands();
-	FUniformToolBarBuilder ToolbarBuilder(CommandList, FMultiBoxCustomization(UVEditorMode->GetModeInfo().ToolbarCustomizationName), TSharedPtr<FExtender>(), false);
-	ToolbarBuilder.SetStyle(&FEditorStyle::Get(), "PaletteToolBar");
-
-	// TODO: Add more tools here
-	ToolbarBuilder.AddToolBarButton(CommandInfos.BeginSelectTool);
-	ToolbarBuilder.AddToolBarButton(CommandInfos.BeginLayoutTool);
-	ToolbarBuilder.AddToolBarButton(CommandInfos.BeginParameterizeMeshTool);
-	ToolbarBuilder.AddToolBarButton(CommandInfos.BeginChannelEditTool);
-	ToolbarBuilder.AddToolBarButton(CommandInfos.BeginSeamTool);
-	ToolbarBuilder.AddToolBarButton(CommandInfos.BeginRecomputeUVsTool);
-
-	// Hook in the tool palette
-	ToolButtonsContainer->SetContent(ToolbarBuilder.MakeWidget());
 
 	// Hook up the tool detail panel
 	ToolDetailsContainer->SetContent(DetailsView.ToSharedRef());
@@ -394,4 +381,35 @@ void FUVEditorModeToolkit::OnToolEnded(UInteractiveToolManager* Manager, UIntera
 	}
 }
 
+
+// Place tool category names here, for creating the tool palette below
+static const FName ToolsTabName(TEXT("Tools"));
+
+const TArray<FName> FUVEditorModeToolkit::PaletteNames_Standard = { ToolsTabName };
+
+void FUVEditorModeToolkit::GetToolPaletteNames(TArray<FName>& PaletteNames) const
+{
+	PaletteNames = PaletteNames_Standard;
+}
+
+FText FUVEditorModeToolkit::GetToolPaletteDisplayName(FName Palette) const
+{
+	return FText::FromName(Palette);
+}
+
+void FUVEditorModeToolkit::BuildToolPalette(FName PaletteIndex, class FToolBarBuilder& ToolbarBuilder)
+{
+	const FUVEditorCommands& Commands = FUVEditorCommands::Get();
+
+	if (PaletteIndex == ToolsTabName)
+	{
+		ToolbarBuilder.AddToolBarButton(Commands.BeginSelectTool);
+		ToolbarBuilder.AddToolBarButton(Commands.BeginLayoutTool);
+		ToolbarBuilder.AddToolBarButton(Commands.BeginChannelEditTool);
+		ToolbarBuilder.AddToolBarButton(Commands.BeginSeamTool);
+		ToolbarBuilder.AddToolBarButton(Commands.BeginParameterizeMeshTool);
+		ToolbarBuilder.AddToolBarButton(Commands.BeginRecomputeUVsTool);
+	}
+}
+
 #undef LOCTEXT_NAMESPACE
diff --git a/Engine/Plugins/Experimental/UVEditor/Source/UVEditor/Private/UVEditorModeUILayer.cpp b/Engine/Plugins/Experimental/UVEditor/Source/UVEditor/Private/UVEditorModeUILayer.cpp
new file mode 100644
--- /dev/null
+++ b/Engine/Plugins/Experimental/UVEditor/Source/UVEditor/Private/UVEditorModeUILayer.cpp
@@ -0,0 +1,55 @@
+// Copyright Epic Games, Inc. All Rights Reserved.
+
+#include "UVEditorModeUILayer.h"
+#include "WorkspaceMenuStructure.h"
+#include "WorkspaceMenuStructureModule.h"
+#include "Toolkits/IToolkit.h"
+#include "UVEditorToolkit.h"
+#include "UVEditorModule.h"
+
+void UUVEditorUISubsystem::Initialize(FSubsystemCollectionBase& Collection)
+{
+	FUVEditorModule& UVEditorModule = FModuleManager::GetModuleChecked<FUVEditorModule>("UVEditor");
+	UVEditorModule.OnRegisterLayoutExtensions().AddUObject(this, &UUVEditorUISubsystem::RegisterLayoutExtensions);
+}
+void UUVEditorUISubsystem::Deinitialize()
+{
+	FUVEditorModule& UVEditorModule = FModuleManager::GetModuleChecked<FUVEditorModule>("UVEditor");
+	UVEditorModule.OnRegisterLayoutExtensions().RemoveAll(this);
+}
+void UUVEditorUISubsystem::RegisterLayoutExtensions(FLayoutExtender& Extender)
+{
+	FTabManager::FTab NewTab(FTabId(UAssetEditorUISubsystem::TopLeftTabID, ETabIdFlags::SaveLayout), ETabState::ClosedTab);
+	Extender.ExtendStack("ToolbarArea", ELayoutExtensionPosition::After, NewTab);
+}
+
+FUVEditorModeUILayer::FUVEditorModeUILayer(const IToolkitHost* InToolkitHost) :
+	FAssetEditorModeUILayer(InToolkitHost)
+{}
+
+void FUVEditorModeUILayer::OnToolkitHostingStarted(const TSharedRef<IToolkit>& Toolkit)
+{
+	if (!Toolkit->IsAssetEditor())
+	{
+		FAssetEditorModeUILayer::OnToolkitHostingStarted(Toolkit);
+		HostedToolkit = Toolkit;
+		Toolkit->SetModeUILayer(SharedThis(this));
+		Toolkit->RegisterTabSpawners(ToolkitHost->GetTabManager().ToSharedRef());
+		RegisterModeTabSpawners();
+		OnToolkitHostReadyForUI.ExecuteIfBound();
+	}
+}
+
+void FUVEditorModeUILayer::OnToolkitHostingFinished(const TSharedRef<IToolkit>& Toolkit)
+{
+	if (HostedToolkit.IsValid() && HostedToolkit.Pin() == Toolkit)
+	{
+		FAssetEditorModeUILayer::OnToolkitHostingFinished(Toolkit);
+	}
+}
+
+TSharedPtr<FWorkspaceItem> FUVEditorModeUILayer::GetModeMenuCategory() const
+{
+	const IWorkspaceMenuStructure& MenuStructure = WorkspaceMenu::GetMenuStructure();
+	return MenuStructure.GetLevelEditorModesCategory();
+}
\ No newline at end of file
diff --git a/Engine/Plugins/Experimental/UVEditor/Source/UVEditor/Private/UVEditorToolkit.cpp b/Engine/Plugins/Experimental/UVEditor/Source/UVEditor/Private/UVEditorToolkit.cpp
--- a/Engine/Plugins/Experimental/UVEditor/Source/UVEditor/Private/UVEditorToolkit.cpp
+++ b/Engine/Plugins/Experimental/UVEditor/Source/UVEditor/Private/UVEditorToolkit.cpp
@@ -22,8 +22,10 @@
 #include "UVEditorMode.h"
 #include "UVEditorModeToolkit.h"
 #include "UVEditorSubsystem.h"
+#include "UVEditorModule.h"
 #include "UVEditorStyle.h"
 #include "UVToolContextObjects.h"
+#include "UVEditorModeUILayer.h"
 #include "Widgets/Docking/SDockTab.h"
 
 #include "SLevelViewport.h"
@@ -61,25 +63,29 @@ FUVEditorToolkit::FUVEditorToolkit(UAssetEditor* InOwningAssetEditor)
 	// We will replace the StandaloneDefaultLayout that our parent class gave us with 
 	// one where the properties detail panel is a vertical column on the left, and there
 	// are two viewports on the right.
+	// We define explicit ExtensionIds on the stacks to reference them later when the
+    // UILayer provides layout extensions. 
+
 	StandaloneDefaultLayout = FTabManager::NewLayout(FName("UVEditorLayout1"))
 		->AddArea
 		(
 			FTabManager::NewPrimaryArea()->SetOrientation(Orient_Vertical)
 			->Split
 			(
-				FTabManager::NewSplitter()->SetOrientation(Orient_Horizontal)
+				FTabManager::NewSplitter()->SetOrientation(Orient_Horizontal)				
 				->Split
 				(
 					FTabManager::NewStack()
-					->SetSizeCoefficient(0.2f)
-					->AddTab(InteractiveToolsPanelTabID, ETabState::OpenedTab)
-					->SetHideTabWell(false)
-				)
+					->SetSizeCoefficient(0.2f)					
+					->SetExtensionId("ToolbarArea")
+					->SetHideTabWell(true)
+				)				
 				->Split
 				(
 					FTabManager::NewStack()
 					->SetSizeCoefficient(0.4f)
 					->AddTab(ViewportTabID, ETabState::OpenedTab)
+					->SetExtensionId("Viewport2DArea")
 					->SetHideTabWell(true)
 				)
 				->Split
@@ -87,10 +93,19 @@ FUVEditorToolkit::FUVEditorToolkit(UAssetEditor* InOwningAssetEditor)
 					FTabManager::NewStack()
 					->SetSizeCoefficient(0.4f)
 					->AddTab(LivePreviewTabID, ETabState::OpenedTab)
+					->SetExtensionId("Viewport3DArea")
 					->SetHideTabWell(true)
 				)
 			)
 		);
+		
+	// Add any extenders specified by the UStaticMeshEditorUISubsystem
+    // The extenders provide defined locations for FModeTookit to attach
+    // tool palette tabs and detail panel tabs
+	LayoutExtender = MakeShared<FLayoutExtender>();
+	FUVEditorModule* UVEditorModule = &FModuleManager::LoadModuleChecked<FUVEditorModule>("UVEditor");
+	UVEditorModule->OnRegisterLayoutExtensions().Broadcast(*LayoutExtender);
+	StandaloneDefaultLayout->ProcessExtensions(*LayoutExtender);
 
 	// We could create the preview scenes in CreateEditorViewportClient() the way that FBaseAssetToolkit
 	// does, but it seems more intuitive to create them right off the bat and pass it in later. 
@@ -182,13 +197,9 @@ void FUVEditorToolkit::RegisterTabSpawners(const TSharedRef<FTabManager>& InTabM
 	FAssetEditorToolkit::RegisterTabSpawners(InTabManager);
 
 	// Here we set up the tabs we referenced in StandaloneDefaultLayout (in the constructor).
-
-	InTabManager->RegisterTabSpawner(InteractiveToolsPanelTabID, FOnSpawnTab::CreateSP(this, 
-		&FUVEditorToolkit::SpawnTab_InteractiveToolsPanel))
-		.SetDisplayName(LOCTEXT("UVToolsPanelLabel", "UV Tools Panel"))
-		.SetGroup(AssetEditorTabsCategory.ToSharedRef())
-		.SetIcon(FSlateIcon(FEditorStyle::GetStyleSetName(), "LevelEditor.Tabs.Details"));
-
+	// We don't deal with the toolbar palette here, since this is handled by existing
+    // infrastructure in FModeToolkit. We only setup spawners for our custom tabs, namely
+    // the 2D and 3D viewports.
 	InTabManager->RegisterTabSpawner(ViewportTabID, FOnSpawnTab::CreateSP(this, &FUVEditorToolkit::SpawnTab_Viewport))
 		.SetDisplayName(LOCTEXT("2DViewportTabLabel", "2D Viewport"))
 		.SetGroup(AssetEditorTabsCategory.ToSharedRef())
@@ -211,6 +222,7 @@ bool FUVEditorToolkit::OnRequestClose()
 	UUVEditorMode* UVMode = Cast<UUVEditorMode>(EditorModeManager->GetActiveScriptableMode(UUVEditorMode::EM_UVEditorModeId));
 	check(UVMode);
 
+	bool bAllowClose = true;
 	// Warn the user of any unapplied changes.
 	if (UVMode->HaveUnappliedChanges())
 	{
@@ -224,18 +236,30 @@ bool FUVEditorToolkit::OnRequestClose()
 		{
 		case EAppReturnType::Yes:
 			UVMode->ApplyChanges();
+			bAllowClose = true;
 			break;
 
 		case EAppReturnType::No:
 			// exit
+			bAllowClose = true;
 			break;
 
 		case EAppReturnType::Cancel:
 			// don't exit
+			bAllowClose = false;
 			return false;
 		}
 	}
 
+	// Give any active modes a chance to shutdown while the toolkit host is still alive
+    // This is super important to do, otherwise currently opened tabs won't be marked as "closed".
+    // This results in tabs not being properly recycled upon reopening the editor and tab
+    // duplication for each opening event.
+	if (bAllowClose)
+	{
+		GetEditorModeManager().ActivateDefaultMode();
+	}
+
 	return FAssetEditorToolkit::OnRequestClose();
 }
 
@@ -351,6 +375,13 @@ void FUVEditorToolkit::CreateEditorModeManager()
 
 void FUVEditorToolkit::PostInitAssetEditor()
 {
+
+    // We setup the ModeUILayer connection here, since the InitAssetEditor is closed off to us. 
+    // Other editors perform this step elsewhere, but this is our best location.
+	TSharedPtr<class IToolkitHost> PinnedToolkitHost = ToolkitHost.Pin();
+	check(PinnedToolkitHost.IsValid());
+	ModeUILayer = MakeShareable(new FUVEditorModeUILayer(PinnedToolkitHost.Get()));
+
 	// Currently, aside from setting up all the UI elements, the toolkit also kicks off the UV
 	// editor mode, which is the mode that the editor always works in (things are packaged into
 	// a mode so that they can be moved to another asset editor if necessary).
@@ -389,23 +420,15 @@ void FUVEditorToolkit::PostInitAssetEditor()
 	UVMode->InitializeTargets(ObjectsToEdit, ObjectTransforms);
 
 	// Regardless of how the user has modified the layout, we're going to make sure that we have
-	// a tool panel to show the tools, and a 2d viewport that will allow our mode to receive ticks.
-	TSharedPtr<SDockTab> ToolsPanel = TabManager->FindExistingLiveTab(InteractiveToolsPanelTabID);
-	if (!ToolsPanel)
-	{
-		ToolsPanel = TabManager->TryInvokeTab(InteractiveToolsPanelTabID);
-	}
+	// a 2d viewport that will allow our mode to receive ticks.
+    // We don't need to invoke the tool palette tab anymore, since this is handled by
+    // underlying infrastructure.
 	if (!TabManager->FindExistingLiveTab(ViewportTabID))
 	{
 		TabManager->TryInvokeTab(ViewportTabID);
 	}
 
-	// Plug in the mode tool panel
 	TSharedPtr<FUVEditorModeToolkit> UVModeToolkit = StaticCastSharedPtr<FUVEditorModeToolkit>(UVMode->GetToolkit().Pin());
-	if (ensure(ToolsPanel && UVModeToolkit.IsValid()))
-	{
-		ToolsPanel->SetContent(UVModeToolkit->GetInlineContent().ToSharedRef());
-	}
 
 	// Add the "Apply Changes" button. It should actually be safe to do this almost
 	// any time, even before that toolbar's registration, but it's easier to put most
@@ -502,4 +525,14 @@ void FUVEditorToolkit::PostInitAssetEditor()
 	LivePreviewViewportClient->SetLookAtLocation(FVector(0, 0, 0));
 }
 
+void FUVEditorToolkit::OnToolkitHostingStarted(const TSharedRef<IToolkit>& Toolkit)
+{
+	ModeUILayer->OnToolkitHostingStarted(Toolkit);
+}
+
+void FUVEditorToolkit::OnToolkitHostingFinished(const TSharedRef<IToolkit>& Toolkit)
+{
+	ModeUILayer->OnToolkitHostingFinished(Toolkit);
+}
+
 #undef LOCTEXT_NAMESPACE
\ No newline at end of file
diff --git a/Engine/Plugins/Experimental/UVEditor/Source/UVEditor/Public/UVEditorModeToolkit.h b/Engine/Plugins/Experimental/UVEditor/Source/UVEditor/Public/UVEditorModeToolkit.h
--- a/Engine/Plugins/Experimental/UVEditor/Source/UVEditor/Public/UVEditorModeToolkit.h
+++ b/Engine/Plugins/Experimental/UVEditor/Source/UVEditor/Public/UVEditorModeToolkit.h
@@ -31,12 +31,26 @@ public:
 	virtual void Init(const TSharedPtr<IToolkitHost>& InitToolkitHost, TWeakObjectPtr<UEdMode> InOwningMode) override;
 	virtual FText GetActiveToolDisplayName() const override { return ActiveToolName; }
 
+	/** Returns the Mode specific tabs in the mode toolbar **/
+	virtual void GetToolPaletteNames(TArray<FName>& InPaletteName) const;
+
+    /** Returns human readable name for the specific palette tool category **/
+	virtual FText GetToolPaletteDisplayName(FName PaletteName) const;
+
+    /** For a specific tool palette category, construct and fill ToolbarBuilder with the category's tools **/
+	virtual void BuildToolPalette(FName PaletteName, class FToolBarBuilder& ToolbarBuilder);
+
+	/** Informs the built in FModeToolkit palette building routines about how we want the toolbars to be setup **/
+	virtual bool HasIntegratedToolPalettes() const { return true; }
+	virtual bool HasExclusiveToolPalettes() const { return false; }
+
 	// IToolkit
 	virtual FName GetToolkitFName() const override;
 	virtual FText GetBaseToolkitName() const override;
-	virtual TSharedPtr<class SWidget> GetInlineContent() const override { return ToolkitWidget; }
 
 protected:
+	const static TArray<FName> PaletteNames_Standard;
+
 	// FModeToolkit
 	virtual void OnToolStarted(UInteractiveToolManager* Manager, UInteractiveTool* Tool) override;
 	virtual void OnToolEnded(UInteractiveToolManager* Manager, UInteractiveTool* Tool) override;
diff --git a/Engine/Plugins/Experimental/UVEditor/Source/UVEditor/Public/UVEditorModeUILayer.h b/Engine/Plugins/Experimental/UVEditor/Source/UVEditor/Public/UVEditorModeUILayer.h
new file mode 100644
--- /dev/null
+++ b/Engine/Plugins/Experimental/UVEditor/Source/UVEditor/Public/UVEditorModeUILayer.h
@@ -0,0 +1,28 @@
+// Copyright Epic Games, Inc. All Rights Reserved.
+
+#pragma once
+
+#include "Toolkits/AssetEditorModeUILayer.h"
+#include "UVEditorModeUILayer.generated.h"
+
+/** Interchange layer to manage built in tab locations within the editor's layout. **/
+UCLASS()
+class UUVEditorUISubsystem : public UAssetEditorUISubsystem
+{
+	GENERATED_BODY()
+public:
+	virtual void Initialize(FSubsystemCollectionBase& Collection) override;
+	virtual void Deinitialize() override;
+	virtual void RegisterLayoutExtensions(FLayoutExtender& Extender) override;
+};
+
+/** Handles the hosting of additional toolkits, such as the mode toolkit, within the UVEditor's toolkit. **/
+class FUVEditorModeUILayer : public FAssetEditorModeUILayer
+{
+public:
+	FUVEditorModeUILayer(const IToolkitHost* InToolkitHost);
+	void OnToolkitHostingStarted(const TSharedRef<IToolkit>& Toolkit) override;
+	void OnToolkitHostingFinished(const TSharedRef<IToolkit>& Toolkit) override;
+
+	TSharedPtr<FWorkspaceItem> GetModeMenuCategory() const override;
+};
diff --git a/Engine/Plugins/Experimental/UVEditor/Source/UVEditor/Public/UVEditorModule.h b/Engine/Plugins/Experimental/UVEditor/Source/UVEditor/Public/UVEditorModule.h
--- a/Engine/Plugins/Experimental/UVEditor/Source/UVEditor/Public/UVEditorModule.h
+++ b/Engine/Plugins/Experimental/UVEditor/Source/UVEditor/Public/UVEditorModule.h
@@ -17,6 +17,12 @@ public:
 	virtual void StartupModule() override;
 	virtual void ShutdownModule() override;
 
+	DECLARE_EVENT_OneParam(FUVEditorModule, FOnRegisterLayoutExtensions, FLayoutExtender&);
+	virtual FOnRegisterLayoutExtensions& OnRegisterLayoutExtensions() { return RegisterLayoutExtensions; }
+
+private:
+	FOnRegisterLayoutExtensions	RegisterLayoutExtensions;
+
 protected:
 	void RegisterMenus();
 };
diff --git a/Engine/Plugins/Experimental/UVEditor/Source/UVEditor/Public/UVEditorToolkit.h b/Engine/Plugins/Experimental/UVEditor/Source/UVEditor/Public/UVEditorToolkit.h
--- a/Engine/Plugins/Experimental/UVEditor/Source/UVEditor/Public/UVEditorToolkit.h
+++ b/Engine/Plugins/Experimental/UVEditor/Source/UVEditor/Public/UVEditorToolkit.h
@@ -14,6 +14,7 @@ class SBorder;
 class SWidget;
 class UInteractiveToolsContext;
 class UInputRouter;
+class FUVEditorModeUILayer;
 
 /**
  * The toolkit is supposed to act as the UI manager for the asset editor. It's responsible 
@@ -48,6 +49,8 @@ public:
 	virtual void RegisterTabSpawners(const TSharedRef<FTabManager>& InTabManager) override;
 	virtual bool OnRequestClose() override;
 	virtual void SaveAsset_Execute() override;
+	void OnToolkitHostingStarted(const TSharedRef<IToolkit>& Toolkit) override;
+	void OnToolkitHostingFinished(const TSharedRef<IToolkit>& Toolkit) override;
 
 	// IAssetEditorInstance
 	// This is important because if this returns true, attempting to edit a static mesh that is
@@ -80,4 +83,6 @@ protected:
 	
 	TSharedPtr<FAssetEditorModeManager> LivePreviewEditorModeManager;
 	TObjectPtr<UInputRouter> LivePreviewInputRouter = nullptr;
+
+	TSharedPtr<FUVEditorModeUILayer> ModeUILayer;
 };
\ No newline at end of file
diff --git a/Engine/Plugins/Experimental/UVEditor/Source/UVEditor/UVEditor.Build.cs b/Engine/Plugins/Experimental/UVEditor/Source/UVEditor/UVEditor.Build.cs
--- a/Engine/Plugins/Experimental/UVEditor/Source/UVEditor/UVEditor.Build.cs
+++ b/Engine/Plugins/Experimental/UVEditor/Source/UVEditor/UVEditor.Build.cs
@@ -59,7 +59,9 @@ public class UVEditor : ModuleRules
 				"ToolMenus",
 				"UVEditorTools",
 				"UVEditorToolsEditorOnly",
-				"ToolWidgets",				
+				"ToolWidgets",
+
+				"WorkspaceMenuStructure",
 				// ... add private dependencies that you statically link with here ...	
 			}
 			);
