UVEditor: Add icons to current select tool actions, in preparation for them someday being placed among the tool icons.

#rnx

diff --git a/Engine/Plugins/Experimental/UVEditor/Source/UVEditor/Private/DetailsCustomizations/UVSelectToolCustomizations.cpp b/Engine/Plugins/Experimental/UVEditor/Source/UVEditor/Private/DetailsCustomizations/UVSelectToolCustomizations.cpp
new file mode 100644
--- /dev/null
+++ b/Engine/Plugins/Experimental/UVEditor/Source/UVEditor/Private/DetailsCustomizations/UVSelectToolCustomizations.cpp
@@ -0,0 +1,80 @@
+// Copyright Epic Games, Inc. All Rights Reserved.
+
+#include "DetailsCustomizations/UVSelectToolCustomizations.h"
+
+#include "DetailCategoryBuilder.h"
+#include "DetailLayoutBuilder.h"
+#include "DetailWidgetRow.h"
+#include "UVEditorCommands.h"
+#include "UVSelectTool.h"
+
+#define LOCTEXT_NAMESPACE "UVSelectToolCustomizations"
+
+TSharedRef<IDetailCustomization> FUVSelectToolActionPropertySetDetails::MakeInstance()
+{
+	return MakeShareable(new FUVSelectToolActionPropertySetDetails);
+}
+
+
+void FUVSelectToolActionPropertySetDetails::CustomizeDetails(IDetailLayoutBuilder& DetailBuilder)
+{
+	// Replace the actions category with a toolbar of action buttons that look like the tool buttons.
+	// When we move selection into the mode, we will move these to the same toolbar as the tools.
+	DetailBuilder.HideCategory("Actions");
+
+	// We need the tool so we can route the commands from the new toolbar to it.
+	TArray<TWeakObjectPtr<UObject>> ObjectsBeingCustomized;
+	DetailBuilder.GetObjectsBeingCustomized(ObjectsBeingCustomized);
+	if (!ensure(ObjectsBeingCustomized.Num() > 0))
+	{
+		return;
+	}
+	USelectToolActionPropertySet* ActionsObject = Cast<USelectToolActionPropertySet>(ObjectsBeingCustomized[0]);
+	if (!ensure(ActionsObject))
+	{
+		return;
+	}
+	UUVSelectTool* Tool = ActionsObject->ParentTool.Get();
+	if (!ensure(Tool))
+	{
+		return;
+	}
+
+	// We have to use some command list since we're creating the toolbar in this FUICommandInfo-based way.
+	// We can't use the actual one that the tool uses because there's not a way to get to it through the
+	// tool. This means that these can't yet be bound to hotkeys, but that's ok because they will be moving
+	// out to mode level at some point anyway, and at that point we will have more flexibility.
+	TSharedPtr<FUICommandList> CommandList = MakeShared<FUICommandList>();
+	const FUVEditorCommands& CommandInfos = FUVEditorCommands::Get();
+
+	auto MapAction = [CommandList, Tool](TSharedPtr<const FUICommandInfo> CommandInfo, ESelectToolAction Action)
+	{
+		CommandList->MapAction(
+			CommandInfo,
+			FExecuteAction::CreateLambda([Tool, Action]() {
+				Tool->RequestAction(Action);
+			}),
+			FCanExecuteAction(),
+			FGetActionCheckState(),
+			FIsActionButtonVisible(),
+			EUIActionRepeatMode::RepeatDisabled
+		);
+	};
+	MapAction(CommandInfos.SewAction, ESelectToolAction::Sew);
+	MapAction(CommandInfos.SplitAction, ESelectToolAction::Split);
+	MapAction(CommandInfos.IslandConformalUnwrapAction, ESelectToolAction::IslandConformalUnwrap);
+
+	// Finally build the actual toolbar
+	FUniformToolBarBuilder ToolbarBuilder(CommandList, FMultiBoxCustomization(NAME_None));
+	ToolbarBuilder.AddToolBarButton(CommandInfos.SewAction);
+	ToolbarBuilder.AddToolBarButton(CommandInfos.SplitAction);
+	ToolbarBuilder.AddToolBarButton(CommandInfos.IslandConformalUnwrapAction);
+
+	DetailBuilder.EditCategory("SelectionActions")
+		.AddCustomRow(LOCTEXT("ActionsSectionFilterString", "Selection Actions"), false)
+		[
+			ToolbarBuilder.MakeWidget()
+		];
+}
+
+#undef LOCTEXT_NAMESPACE
\ No newline at end of file
diff --git a/Engine/Plugins/Experimental/UVEditor/Source/UVEditor/Private/UVEditorCommands.cpp b/Engine/Plugins/Experimental/UVEditor/Source/UVEditor/Private/UVEditorCommands.cpp
--- a/Engine/Plugins/Experimental/UVEditor/Source/UVEditor/Private/UVEditorCommands.cpp
+++ b/Engine/Plugins/Experimental/UVEditor/Source/UVEditor/Private/UVEditorCommands.cpp
@@ -31,16 +31,22 @@ void FUVEditorCommands::RegisterCommands()
 	UI_COMMAND(BeginSeamTool, "Seam", "Switch to Seam tool", EUserInterfaceActionType::ToggleButton, FInputChord());
 	UI_COMMAND(BeginRecomputeUVsTool, "Unwrap", "Switch to Unwrap tool", EUserInterfaceActionType::ToggleButton, FInputChord());
 
+	// These currently get linked to actions inside the select tool, but will eventually have their own buttons among the tools
+	// once selection is pulled out to mode-level.
+	UI_COMMAND(SewAction, "Sew", "Sew selected edges to their pairs", EUserInterfaceActionType::Button, FInputChord());
+	UI_COMMAND(SplitAction, "Split", "Split UV's at selected edges", EUserInterfaceActionType::Button, FInputChord());
+	UI_COMMAND(IslandConformalUnwrapAction, "IsdUw", "Perform a  conformal unwrap of selected islands", EUserInterfaceActionType::Button, FInputChord());
+
 	// These allow us to link up to pressed keys
 	UI_COMMAND(AcceptOrCompleteActiveTool, "Accept", "Accept the active tool", EUserInterfaceActionType::Button, FInputChord(EKeys::Enter));
 	UI_COMMAND(CancelOrCompleteActiveTool, "Cancel", "Cancel the active tool", EUserInterfaceActionType::Button, FInputChord(EKeys::Escape));
 
 	// These get used in viewport buttons
-	UI_COMMAND(VertexSelection, "Vertex Selection", "Select vertex on click", EUserInterfaceActionType::ToggleButton, FInputChord());
-	UI_COMMAND(EdgeSelection, "Edge Selection", "Select edge on click", EUserInterfaceActionType::ToggleButton, FInputChord());
-	UI_COMMAND(TriangleSelection, "Triangle Selection", "Select triangle on click", EUserInterfaceActionType::ToggleButton, FInputChord());
-	UI_COMMAND(IslandSelection, "Island Selection", "Select connected island on click", EUserInterfaceActionType::ToggleButton, FInputChord());
-	UI_COMMAND(FullMeshSelection, "Full Mesh Selection", "Select whole mesh on click", EUserInterfaceActionType::ToggleButton, FInputChord());
+	UI_COMMAND(VertexSelection, "Vertex Selection", "Select vertices", EUserInterfaceActionType::ToggleButton, FInputChord());
+	UI_COMMAND(EdgeSelection, "Edge Selection", "Select edges", EUserInterfaceActionType::ToggleButton, FInputChord());
+	UI_COMMAND(TriangleSelection, "Triangle Selection", "Select triangles", EUserInterfaceActionType::ToggleButton, FInputChord());
+	UI_COMMAND(IslandSelection, "Island Selection", "Select connected islands", EUserInterfaceActionType::ToggleButton, FInputChord());
+	UI_COMMAND(FullMeshSelection, "Full Mesh Selection", "Select whole mesh", EUserInterfaceActionType::ToggleButton, FInputChord());
 
 	UI_COMMAND(EnableOrbitCamera, "Orbit", "Enable Orbit Camera", EUserInterfaceActionType::ToggleButton, FInputChord());
 	UI_COMMAND(EnableFlyCamera, "Fly", "Enable Fly Camera", EUserInterfaceActionType::ToggleButton, FInputChord());
diff --git a/Engine/Plugins/Experimental/UVEditor/Source/UVEditor/Private/UVEditorModule.cpp b/Engine/Plugins/Experimental/UVEditor/Source/UVEditor/Private/UVEditorModule.cpp
--- a/Engine/Plugins/Experimental/UVEditor/Source/UVEditor/Private/UVEditorModule.cpp
+++ b/Engine/Plugins/Experimental/UVEditor/Source/UVEditor/Private/UVEditorModule.cpp
@@ -3,14 +3,17 @@
 #include "UVEditorModule.h"
 
 #include "ContentBrowserMenuContexts.h"
+#include "DetailsCustomizations/UVSelectToolCustomizations.h"
+#include "EditorModeRegistry.h"
 #include "LevelEditorMenuContext.h"
+#include "PropertyEditorModule.h"
 #include "Selection.h"
 #include "UVEditor.h"
 #include "UVEditorCommands.h"
 #include "UVEditorMode.h"
 #include "UVEditorStyle.h"
 #include "UVEditorSubsystem.h"
-#include "EditorModeRegistry.h"
+#include "UVSelectTool.h"
 
 #include "ToolMenus.h"
 
@@ -24,6 +27,11 @@ void FUVEditorModule::StartupModule()
 	// Menus need to be registered in a callback to make sure the system is ready for them.
 	UToolMenus::RegisterStartupCallback(
 		FSimpleMulticastDelegate::FDelegate::CreateRaw(this, &FUVEditorModule::RegisterMenus));
+
+	// Register details view customizations
+	FPropertyEditorModule& PropertyModule = FModuleManager::LoadModuleChecked<FPropertyEditorModule>("PropertyEditor");
+	PropertyModule.RegisterCustomClassLayout(USelectToolActionPropertySet::StaticClass()->GetFName(), 
+		FOnGetDetailCustomizationInstance::CreateStatic(&FUVSelectToolActionPropertySetDetails::MakeInstance));
 }
 
 void FUVEditorModule::ShutdownModule()
@@ -35,6 +43,10 @@ void FUVEditorModule::ShutdownModule()
 	FUVEditorCommands::Unregister();
 
 	FEditorModeRegistry::Get().UnregisterMode(UUVEditorMode::EM_UVEditorModeId);
+
+	// Unregister customizations
+	FPropertyEditorModule* PropertyEditorModule = FModuleManager::GetModulePtr<FPropertyEditorModule>("PropertyEditor");
+	PropertyEditorModule->UnregisterCustomClassLayout(USelectToolActionPropertySet::StaticClass()->GetFName());
 }
 
 void FUVEditorModule::RegisterMenus()
diff --git a/Engine/Plugins/Experimental/UVEditor/Source/UVEditor/Private/UVEditorStyle.cpp b/Engine/Plugins/Experimental/UVEditor/Source/UVEditor/Private/UVEditorStyle.cpp
--- a/Engine/Plugins/Experimental/UVEditor/Source/UVEditor/Private/UVEditorStyle.cpp
+++ b/Engine/Plugins/Experimental/UVEditor/Source/UVEditor/Private/UVEditorStyle.cpp
@@ -29,6 +29,11 @@ FUVEditorStyle::FUVEditorStyle()
 	Set("UVEditor.BeginSeamTool",               new IMAGE_BRUSH_SVG("ModelingUVSeamEdit", ToolbarIconSize));
 	Set("UVEditor.BeginRecomputeUVsTool",       new IMAGE_BRUSH_SVG("GroupUnwrap", ToolbarIconSize));
 
+	// Select tool actions
+	Set("UVEditor.SewAction", new IMAGE_BRUSH_SVG("UVSew", ToolbarIconSize));
+	Set("UVEditor.SplitAction", new IMAGE_BRUSH_SVG("UVCut", ToolbarIconSize));
+	Set("UVEditor.IslandConformalUnwrapAction", new IMAGE_BRUSH_SVG("UVUnwrap", ToolbarIconSize));
+
 	// Top toolbar icons
 	Set("UVEditor.ApplyChanges", new CORE_IMAGE_BRUSH_SVG("Starship/Common/Apply", ToolbarIconSize));
 	Set("UVEditor.ChannelSettings", new CORE_IMAGE_BRUSH_SVG("Starship/Common/SetDrawUVs", ToolbarIconSize));
diff --git a/Engine/Plugins/Experimental/UVEditor/Source/UVEditor/Public/DetailsCustomizations/UVSelectToolCustomizations.h b/Engine/Plugins/Experimental/UVEditor/Source/UVEditor/Public/DetailsCustomizations/UVSelectToolCustomizations.h
new file mode 100644
--- /dev/null
+++ b/Engine/Plugins/Experimental/UVEditor/Source/UVEditor/Public/DetailsCustomizations/UVSelectToolCustomizations.h
@@ -0,0 +1,12 @@
+// Copyright Epic Games, Inc. All Rights Reserved.
+
+#pragma once
+
+#include "IDetailCustomization.h"
+
+class FUVSelectToolActionPropertySetDetails : public IDetailCustomization
+{
+public:
+	static TSharedRef<IDetailCustomization> MakeInstance();
+	void CustomizeDetails(IDetailLayoutBuilder& DetailBuilder) override;
+};
diff --git a/Engine/Plugins/Experimental/UVEditor/Source/UVEditor/Public/UVEditorCommands.h b/Engine/Plugins/Experimental/UVEditor/Source/UVEditor/Public/UVEditorCommands.h
--- a/Engine/Plugins/Experimental/UVEditor/Source/UVEditor/Public/UVEditorCommands.h
+++ b/Engine/Plugins/Experimental/UVEditor/Source/UVEditor/Public/UVEditorCommands.h
@@ -25,6 +25,10 @@ public:
 	TSharedPtr<FUICommandInfo> BeginSeamTool;
 	TSharedPtr<FUICommandInfo> BeginRecomputeUVsTool;
 
+	TSharedPtr<FUICommandInfo> SewAction;
+	TSharedPtr<FUICommandInfo> SplitAction;
+	TSharedPtr<FUICommandInfo> IslandConformalUnwrapAction;
+
 	TSharedPtr<FUICommandInfo> AcceptOrCompleteActiveTool;
 	TSharedPtr<FUICommandInfo> CancelOrCompleteActiveTool;
 
