UV Editor: Fix NonUnity build error

diff --git a/Engine/Plugins/Experimental/UVEditor/Source/UVEditor/Private/DetailsCustomizations/UVSelectToolCustomizations.cpp b/Engine/Plugins/Experimental/UVEditor/Source/UVEditor/Private/DetailsCustomizations/UVSelectToolCustomizations.cpp
--- a/Engine/Plugins/Experimental/UVEditor/Source/UVEditor/Private/DetailsCustomizations/UVSelectToolCustomizations.cpp
+++ b/Engine/Plugins/Experimental/UVEditor/Source/UVEditor/Private/DetailsCustomizations/UVSelectToolCustomizations.cpp
@@ -1,6 +1,7 @@
 // Copyright Epic Games, Inc. All Rights Reserved.
 
 #include "DetailsCustomizations/UVSelectToolCustomizations.h"
+#include <Framework/MultiBox/MultiBoxBuilder.h>
 
 #include "DetailCategoryBuilder.h"
 #include "DetailLayoutBuilder.h"
diff --git a/Engine/Plugins/Experimental/UVEditor/Source/UVEditor/Public/UVEditorModule.h b/Engine/Plugins/Experimental/UVEditor/Source/UVEditor/Public/UVEditorModule.h
--- a/Engine/Plugins/Experimental/UVEditor/Source/UVEditor/Public/UVEditorModule.h
+++ b/Engine/Plugins/Experimental/UVEditor/Source/UVEditor/Public/UVEditorModule.h
@@ -3,7 +3,9 @@
 #pragma once
 
 #include "CoreMinimal.h"
-#include "Modules/ModuleManager.h"
+#include "Modules/ModuleInterface.h"
+
+class FLayoutExtender;
 
 /**
  * Besides the normal module things, the module class is also responsible for hooking the 
