UVEditor: Cleanup how Operator Factories are referenced by tools, preventing crash bugs on garbage collection. Also adds undo transactions around tools to capture all accepted results for all assets under one undo history item.

#rnx

diff --git a/Engine/Plugins/Experimental/MeshModelingToolsetExp/Source/MeshModelingToolsExp/Private/RecomputeUVsTool.cpp b/Engine/Plugins/Experimental/MeshModelingToolsetExp/Source/MeshModelingToolsExp/Private/RecomputeUVsTool.cpp
--- a/Engine/Plugins/Experimental/MeshModelingToolsetExp/Source/MeshModelingToolsExp/Private/RecomputeUVsTool.cpp
+++ b/Engine/Plugins/Experimental/MeshModelingToolsetExp/Source/MeshModelingToolsExp/Private/RecomputeUVsTool.cpp
@@ -42,24 +42,22 @@ void URecomputeUVsTool::Setup()
 	FComponentMaterialSet MaterialSet = UE::ToolTarget::GetMaterialSet(Target);
 	FTransform TargetTransform = (FTransform)UE::ToolTarget::GetLocalToWorldTransform(Target);
 
-	UE::ToolTarget::HideSourceObject(Target);
-
 	// initialize our properties
 
+	Settings = NewObject<URecomputeUVsToolProperties>(this);
+	Settings->RestoreProperties(this);
+	AddToolPropertySource(Settings);
+
 	UVChannelProperties = NewObject<UMeshUVChannelProperties>(this);
 	UVChannelProperties->RestoreProperties(this);
 	UVChannelProperties->Initialize(InputMesh.Get(), false);
 	UVChannelProperties->ValidateSelection(true);
-	UVChannelProperties->WatchProperty(UVChannelProperties->UVChannel, [this](const FString& NewValue) 
-	{
-		MaterialSettings->UVChannel = UVChannelProperties->GetSelectedChannelIndex(true);
-	});
+	UVChannelProperties->WatchProperty(UVChannelProperties->UVChannel, [this](const FString& NewValue)
+		{
+			MaterialSettings->UVChannel = UVChannelProperties->GetSelectedChannelIndex(true);
+		});
 	AddToolPropertySource(UVChannelProperties);
 
-	Settings = NewObject<URecomputeUVsToolProperties>(this);
-	Settings->RestoreProperties(this);
-	AddToolPropertySource(Settings);
-
 	PolygroupLayerProperties = NewObject<UPolygroupLayersProperties>(this);
 	PolygroupLayerProperties->RestoreProperties(this, TEXT("RecomputeUVsTool"));
 	PolygroupLayerProperties->InitializeGroupLayers(InputMesh.Get());
@@ -71,6 +69,9 @@ void URecomputeUVsTool::Setup()
 	MaterialSettings->MaterialMode = ESetMeshMaterialMode::Checkerboard;
 	MaterialSettings->RestoreProperties(this, TEXT("ModelingUVTools"));
 	AddToolPropertySource(MaterialSettings);
+
+	UE::ToolTarget::HideSourceObject(Target);
+
 	// force update
 	MaterialSettings->UpdateMaterials();
 
diff --git a/Engine/Plugins/Experimental/MeshModelingToolsetExp/Source/MeshModelingToolsExp/Private/UVLayoutTool.cpp b/Engine/Plugins/Experimental/MeshModelingToolsetExp/Source/MeshModelingToolsExp/Private/UVLayoutTool.cpp
--- a/Engine/Plugins/Experimental/MeshModelingToolsetExp/Source/MeshModelingToolsExp/Private/UVLayoutTool.cpp
+++ b/Engine/Plugins/Experimental/MeshModelingToolsetExp/Source/MeshModelingToolsExp/Private/UVLayoutTool.cpp
@@ -146,24 +146,26 @@ void UUVLayoutTool::UpdateNumPreviews()
 		}
 		Previews.SetNum(TargetNumPreview);
 		OriginalDynamicMeshes.SetNum(TargetNumPreview);
+		Factories.SetNum(TargetNumPreview);
 	}
 	else
 	{
 		OriginalDynamicMeshes.SetNum(TargetNumPreview);
+		Factories.SetNum(TargetNumPreview);
 		for (int32 PreviewIdx = CurrentNumPreview; PreviewIdx < TargetNumPreview; PreviewIdx++)
 		{
 			OriginalDynamicMeshes[PreviewIdx] = MakeShared<FDynamicMesh3, ESPMode::ThreadSafe>();
 			FMeshDescriptionToDynamicMesh Converter;
 			Converter.Convert(TargetMeshProviderInterface(PreviewIdx)->GetMeshDescription(), *OriginalDynamicMeshes[PreviewIdx]);
 
-			UUVLayoutOperatorFactory* OpFactory = NewObject<UUVLayoutOperatorFactory>();
-			OpFactory->OriginalMesh = OriginalDynamicMeshes[PreviewIdx];
-			OpFactory->Settings = BasicProperties;
-			OpFactory->TargetTransform = TargetComponentInterface(PreviewIdx)->GetWorldTransform();
-			OpFactory->GetSelectedUVChannel = [this]() { return GetSelectedUVChannel(); };
+			Factories[PreviewIdx]= NewObject<UUVLayoutOperatorFactory>();
+			Factories[PreviewIdx]->OriginalMesh = OriginalDynamicMeshes[PreviewIdx];
+			Factories[PreviewIdx]->Settings = BasicProperties;
+			Factories[PreviewIdx]->TargetTransform = TargetComponentInterface(PreviewIdx)->GetWorldTransform();
+			Factories[PreviewIdx]->GetSelectedUVChannel = [this]() { return GetSelectedUVChannel(); };
 
-			UMeshOpPreviewWithBackgroundCompute* Preview = Previews.Add_GetRef(NewObject<UMeshOpPreviewWithBackgroundCompute>(OpFactory, "Preview"));
-			Preview->Setup(this->TargetWorld, OpFactory);
+			UMeshOpPreviewWithBackgroundCompute* Preview = Previews.Add_GetRef(NewObject<UMeshOpPreviewWithBackgroundCompute>(Factories[PreviewIdx], "Preview"));
+			Preview->Setup(this->TargetWorld, Factories[PreviewIdx]);
 			ToolSetupUtil::ApplyRenderingConfigurationToPreview(Preview->PreviewMesh, Targets[PreviewIdx]); 
 
 			FComponentMaterialSet MaterialSet;
@@ -211,6 +213,10 @@ void UUVLayoutTool::Shutdown(EToolShutdownType ShutdownType)
 	{
 		GenerateAsset(Results);
 	}
+	for (int32 TargetIndex = 0; TargetIndex < Targets.Num(); ++TargetIndex)
+	{
+		Factories[TargetIndex] = nullptr;
+	}
 }
 
 int32 UUVLayoutTool::GetSelectedUVChannel() const
diff --git a/Engine/Plugins/Experimental/MeshModelingToolsetExp/Source/MeshModelingToolsExp/Public/UVLayoutTool.h b/Engine/Plugins/Experimental/MeshModelingToolsetExp/Source/MeshModelingToolsExp/Public/UVLayoutTool.h
--- a/Engine/Plugins/Experimental/MeshModelingToolsetExp/Source/MeshModelingToolsExp/Public/UVLayoutTool.h
+++ b/Engine/Plugins/Experimental/MeshModelingToolsetExp/Source/MeshModelingToolsExp/Public/UVLayoutTool.h
@@ -22,6 +22,7 @@
 struct FMeshDescription;
 class UDynamicMeshComponent;
 class UUVLayoutProperties;
+class UUVLayoutOperatorFactory;
 PREDECLARE_GEOMETRY(class FDynamicMesh3);
 
 /**
@@ -86,6 +87,9 @@ protected:
 	UPROPERTY()
 	TArray<TObjectPtr<UMeshOpPreviewWithBackgroundCompute>> Previews;
 
+	UPROPERTY()
+	TArray<TObjectPtr<UUVLayoutOperatorFactory>> Factories;
+
 protected:
 	TArray<TSharedPtr<UE::Geometry::FDynamicMesh3, ESPMode::ThreadSafe>> OriginalDynamicMeshes;
 
diff --git a/Engine/Plugins/Experimental/UVEditor/Source/UVEditorTools/Private/UVEditorLayoutTool.cpp b/Engine/Plugins/Experimental/UVEditor/Source/UVEditorTools/Private/UVEditorLayoutTool.cpp
--- a/Engine/Plugins/Experimental/UVEditor/Source/UVEditorTools/Private/UVEditorLayoutTool.cpp
+++ b/Engine/Plugins/Experimental/UVEditor/Source/UVEditorTools/Private/UVEditorLayoutTool.cpp
@@ -49,15 +49,17 @@ void UUVEditorLayoutTool::Setup()
 	Settings->RestoreProperties(this);
 	AddToolPropertySource(Settings);
 
-	for (TObjectPtr<UUVEditorToolMeshInput> Target : Targets)
+	Factories.SetNum(Targets.Num());
+	for (int32 TargetIndex = 0; TargetIndex < Targets.Num(); ++TargetIndex)
 	{
-		UUVLayoutOperatorFactory* OpFactory = NewObject<UUVLayoutOperatorFactory>();
-		OpFactory->TargetTransform = Target->AppliedPreview->PreviewMesh->GetTransform();
-		OpFactory->Settings = Settings;
-		OpFactory->OriginalMesh = Target->AppliedCanonical;
-		OpFactory->GetSelectedUVChannel = [Target]() { return Target->UVLayerIndex; };
-
-		Target->AppliedPreview->ChangeOpFactory(OpFactory);
+		TObjectPtr<UUVEditorToolMeshInput> Target = Targets[TargetIndex];
+		Factories[TargetIndex] = NewObject<UUVLayoutOperatorFactory>();
+		Factories[TargetIndex]->TargetTransform = Target->AppliedPreview->PreviewMesh->GetTransform();
+		Factories[TargetIndex]->Settings = Settings;
+		Factories[TargetIndex]->OriginalMesh = Target->AppliedCanonical;
+		Factories[TargetIndex]->GetSelectedUVChannel = [Target]() { return Target->UVLayerIndex; };
+
+		Target->AppliedPreview->ChangeOpFactory(Factories[TargetIndex]);
 		Target->AppliedPreview->OnMeshUpdated.AddWeakLambda(this, [Target](UMeshOpPreviewWithBackgroundCompute* Preview) {
 			Target->UpdateUnwrapPreviewFromAppliedPreview();
 		});
@@ -73,16 +75,16 @@ void UUVEditorLayoutTool::Setup()
 void UUVEditorLayoutTool::Shutdown(EToolShutdownType ShutdownType)
 {
 	Settings->SaveProperties(this);
-
 	for (TObjectPtr<UUVEditorToolMeshInput> Target : Targets)
 	{
 		Target->AppliedPreview->OnMeshUpdated.RemoveAll(this);
-		Target->AppliedPreview->ClearOpFactory();
 	}
 
 	if (ShutdownType == EToolShutdownType::Accept)
 	{
 		UUVToolEmitChangeAPI* ChangeAPI = GetToolManager()->GetContextObjectStore()->FindContext<UUVToolEmitChangeAPI>();
+		const FText TransactionName(LOCTEXT("LayoutTransactionName", "Layout Tool"));
+		ChangeAPI->BeginUndoTransaction(TransactionName);
 
 		for (TObjectPtr<UUVEditorToolMeshInput> Target : Targets)
 		{
@@ -103,6 +105,8 @@ void UUVEditorLayoutTool::Shutdown(EToolShutdownType ShutdownType)
 
 			ChangeAPI->EmitToolIndependentUnwrapCanonicalChange(Target, ChangeTracker.EndChange(), LOCTEXT("ApplyLayoutTool", "Layout Tool"));
 		}
+
+		ChangeAPI->EndUndoTransaction();
 	}
 	else
 	{
@@ -113,6 +117,17 @@ void UUVEditorLayoutTool::Shutdown(EToolShutdownType ShutdownType)
 		}
 	}
 
+
+	for (TObjectPtr<UUVEditorToolMeshInput> Target : Targets)
+	{
+		Target->AppliedPreview->ClearOpFactory();
+	}
+
+	for (int32 TargetIndex = 0; TargetIndex < Targets.Num(); ++TargetIndex)
+	{
+		Factories[TargetIndex] = nullptr;
+	}
+
 	Settings = nullptr;
 	Targets.Empty();
 }
diff --git a/Engine/Plugins/Experimental/UVEditor/Source/UVEditorTools/Public/UVEditorLayoutTool.h b/Engine/Plugins/Experimental/UVEditor/Source/UVEditorTools/Public/UVEditorLayoutTool.h
--- a/Engine/Plugins/Experimental/UVEditor/Source/UVEditorTools/Public/UVEditorLayoutTool.h
+++ b/Engine/Plugins/Experimental/UVEditor/Source/UVEditorTools/Public/UVEditorLayoutTool.h
@@ -11,6 +11,7 @@
 
 class UUVEditorToolMeshInput;
 class UUVLayoutProperties;
+class UUVLayoutOperatorFactory;
 
 UCLASS()
 class UVEDITORTOOLS_API UUVEditorLayoutToolBuilder : public UInteractiveToolBuilder
@@ -64,4 +65,7 @@ protected:
 
 	UPROPERTY()
 	TObjectPtr<UUVLayoutProperties> Settings = nullptr;
+
+	UPROPERTY()
+	TArray<TObjectPtr<UUVLayoutOperatorFactory>> Factories;
 };
diff --git a/Engine/Plugins/Experimental/UVEditor/Source/UVEditorToolsEditorOnly/Private/UVEditorParameterizeMeshTool.cpp b/Engine/Plugins/Experimental/UVEditor/Source/UVEditorToolsEditorOnly/Private/UVEditorParameterizeMeshTool.cpp
--- a/Engine/Plugins/Experimental/UVEditor/Source/UVEditorToolsEditorOnly/Private/UVEditorParameterizeMeshTool.cpp
+++ b/Engine/Plugins/Experimental/UVEditor/Source/UVEditorToolsEditorOnly/Private/UVEditorParameterizeMeshTool.cpp
@@ -88,18 +88,20 @@ void UUVEditorParameterizeMeshTool::Setup()
 		Target->AppliedPreview->OverrideMaterial = MaterialSettings->GetActiveOverrideMaterial();
 	}
 
-	for (TObjectPtr<UUVEditorToolMeshInput> Target : Targets)
+	Factories.SetNum(Targets.Num());
+	for (int32 TargetIndex = 0; TargetIndex < Targets.Num(); ++TargetIndex)
 	{
-		UParameterizeMeshOperatorFactory* OpFactory = NewObject<UParameterizeMeshOperatorFactory>();
-		OpFactory->TargetTransform = Target->AppliedPreview->PreviewMesh->GetTransform();
-		OpFactory->Settings = Settings;
-		OpFactory->UVAtlasProperties = UVAtlasProperties;
-		OpFactory->XAtlasProperties = XAtlasProperties;
-		OpFactory->PatchBuilderProperties = PatchBuilderProperties;
-		OpFactory->OriginalMesh = Target->AppliedCanonical;
-		OpFactory->GetSelectedUVChannel = [Target]() { return Target->UVLayerIndex; };
-
-		Target->AppliedPreview->ChangeOpFactory(OpFactory);
+		TObjectPtr<UUVEditorToolMeshInput> Target = Targets[TargetIndex];
+		Factories[TargetIndex] = NewObject<UParameterizeMeshOperatorFactory>();
+		Factories[TargetIndex]->TargetTransform = Target->AppliedPreview->PreviewMesh->GetTransform();
+		Factories[TargetIndex]->Settings = Settings;
+		Factories[TargetIndex]->UVAtlasProperties = UVAtlasProperties;
+		Factories[TargetIndex]->XAtlasProperties = XAtlasProperties;
+		Factories[TargetIndex]->PatchBuilderProperties = PatchBuilderProperties;
+		Factories[TargetIndex]->OriginalMesh = Target->AppliedCanonical;
+		Factories[TargetIndex]->GetSelectedUVChannel = [Target]() { return Target->UVLayerIndex; };
+
+		Target->AppliedPreview->ChangeOpFactory(Factories[TargetIndex]);
 		Target->AppliedPreview->OnMeshUpdated.AddWeakLambda(this, [Target](UMeshOpPreviewWithBackgroundCompute* Preview) {
 			Target->UpdateUnwrapPreviewFromAppliedPreview();
 			});
@@ -161,13 +163,13 @@ void UUVEditorParameterizeMeshTool::Shutdown(EToolShutdownType ShutdownType)
 	for (TObjectPtr<UUVEditorToolMeshInput> Target : Targets)
 	{
 		Target->AppliedPreview->OnMeshUpdated.RemoveAll(this);
-		Target->AppliedPreview->ClearOpFactory();
-		Target->AppliedPreview->OverrideMaterial = nullptr;
 	}
 
 	if (ShutdownType == EToolShutdownType::Accept)
 	{
 		UUVToolEmitChangeAPI* ChangeAPI = GetToolManager()->GetContextObjectStore()->FindContext<UUVToolEmitChangeAPI>();
+		const FText TransactionName(LOCTEXT("ParameterizeMeshTransactionName", "Auto UV Tool"));
+		ChangeAPI->BeginUndoTransaction(TransactionName);
 
 		for (TObjectPtr<UUVEditorToolMeshInput> Target : Targets)
 		{
@@ -184,6 +186,8 @@ void UUVEditorParameterizeMeshTool::Shutdown(EToolShutdownType ShutdownType)
 
 			ChangeAPI->EmitToolIndependentUnwrapCanonicalChange(Target, ChangeTracker.EndChange(), LOCTEXT("ApplyParameterizeMeshTool", "Auto UV Tool"));
 		}
+
+		ChangeAPI->EndUndoTransaction();
 	}
 	else
 	{
@@ -194,6 +198,17 @@ void UUVEditorParameterizeMeshTool::Shutdown(EToolShutdownType ShutdownType)
 		}
 	}
 
+	for (TObjectPtr<UUVEditorToolMeshInput> Target : Targets)
+	{
+		Target->AppliedPreview->ClearOpFactory();
+		Target->AppliedPreview->OverrideMaterial = nullptr;
+	}
+
+	for (int32 TargetIndex = 0; TargetIndex < Targets.Num(); ++TargetIndex)
+	{
+		Factories[TargetIndex] = nullptr;
+	}
+
 	Settings = nullptr;
 	Targets.Empty();
 }
diff --git a/Engine/Plugins/Experimental/UVEditor/Source/UVEditorToolsEditorOnly/Public/UVEditorParameterizeMeshTool.h b/Engine/Plugins/Experimental/UVEditor/Source/UVEditorToolsEditorOnly/Public/UVEditorParameterizeMeshTool.h
--- a/Engine/Plugins/Experimental/UVEditor/Source/UVEditorToolsEditorOnly/Public/UVEditorParameterizeMeshTool.h
+++ b/Engine/Plugins/Experimental/UVEditor/Source/UVEditorToolsEditorOnly/Public/UVEditorParameterizeMeshTool.h
@@ -23,6 +23,7 @@ class UParameterizeMeshToolProperties;
 class UParameterizeMeshToolUVAtlasProperties;
 class UParameterizeMeshToolXAtlasProperties;
 class UParameterizeMeshToolPatchBuilderProperties;
+class UParameterizeMeshOperatorFactory;
 
 UCLASS()
 class UVEDITORTOOLSEDITORONLY_API UUVEditorParameterizeMeshToolBuilder : public UInteractiveToolBuilder
@@ -90,6 +91,9 @@ protected:
 	UPROPERTY()
 	TObjectPtr<UExistingMeshMaterialProperties> MaterialSettings = nullptr;
 
+	UPROPERTY()
+	TArray<TObjectPtr<UParameterizeMeshOperatorFactory>> Factories;
+
 	void OnMethodTypeChanged();
 };
 
diff --git a/Engine/Plugins/Runtime/MeshModelingToolset/Source/ModelingOperators/Public/ParameterizationOps/UVLayoutOp.h b/Engine/Plugins/Runtime/MeshModelingToolset/Source/ModelingOperators/Public/ParameterizationOps/UVLayoutOp.h
--- a/Engine/Plugins/Runtime/MeshModelingToolset/Source/ModelingOperators/Public/ParameterizationOps/UVLayoutOp.h
+++ b/Engine/Plugins/Runtime/MeshModelingToolset/Source/ModelingOperators/Public/ParameterizationOps/UVLayoutOp.h
@@ -55,6 +55,9 @@ public:
 
 /**
  * Can be hooked up to a UMeshOpPreviewWithBackgroundCompute to perform UV Layout operations.
+ *
+ * Inherits from UObject so that it can hold a strong pointer to the settings UObject, which
+ * needs to be a UObject to be displayed in the details panel.
  */
 UCLASS()
 class MODELINGOPERATORS_API UUVLayoutOperatorFactory : public UObject, public UE::Geometry::IDynamicMeshOperatorFactory
diff --git a/Engine/Plugins/Runtime/MeshModelingToolset/Source/ModelingOperatorsEditorOnly/Public/ParameterizationOps/ParameterizeMeshOp.h b/Engine/Plugins/Runtime/MeshModelingToolset/Source/ModelingOperatorsEditorOnly/Public/ParameterizationOps/ParameterizeMeshOp.h
--- a/Engine/Plugins/Runtime/MeshModelingToolset/Source/ModelingOperatorsEditorOnly/Public/ParameterizationOps/ParameterizeMeshOp.h
+++ b/Engine/Plugins/Runtime/MeshModelingToolset/Source/ModelingOperatorsEditorOnly/Public/ParameterizationOps/ParameterizeMeshOp.h
@@ -121,6 +121,9 @@ protected:
 
 /**
  * Can be hooked up to a UMeshOpPreviewWithBackgroundCompute to perform UV parameterization operations.
+ *
+ * Inherits from UObject so that it can hold a strong pointer to the settings UObject, which
+ * needs to be a UObject to be displayed in the details panel.
  */
 UCLASS()
 class MODELINGOPERATORSEDITORONLY_API UParameterizeMeshOperatorFactory : public UObject, public UE::Geometry::IDynamicMeshOperatorFactory
