Made the r.Editor.OpaqueGizmo cvar scene-specific, and disabled dithering in UV editor viewports so that we could use editor compositing to draw the gizmos on top.

#rnx

diff --git a/Engine/Plugins/Experimental/UVEditor/Source/UVEditor/Private/UVEditorToolkit.cpp b/Engine/Plugins/Experimental/UVEditor/Source/UVEditor/Private/UVEditorToolkit.cpp
--- a/Engine/Plugins/Experimental/UVEditor/Source/UVEditor/Private/UVEditorToolkit.cpp
+++ b/Engine/Plugins/Experimental/UVEditor/Source/UVEditor/Private/UVEditorToolkit.cpp
@@ -494,6 +494,9 @@ void FUVEditorToolkit::PostInitAssetEditor()
 		// Disable motion blur effects that cause our renders to "fade in" as things are moved
 		Client->EngineShowFlags.SetTemporalAA(false);
 		Client->EngineShowFlags.SetMotionBlur(false);
+
+		// Disable the dithering of occluded portions of gizmos.
+		Client->EngineShowFlags.SetOpaqueCompositeEditorPrimitives(true);
 	};
 	SetCommonViewportClientOptions(ViewportClient.Get());
 
diff --git a/Engine/Plugins/Experimental/UVEditor/Source/UVEditorTools/Private/UVSelectTool.cpp b/Engine/Plugins/Experimental/UVEditor/Source/UVEditorTools/Private/UVSelectTool.cpp
--- a/Engine/Plugins/Experimental/UVEditor/Source/UVEditorTools/Private/UVSelectTool.cpp
+++ b/Engine/Plugins/Experimental/UVEditor/Source/UVEditorTools/Private/UVSelectTool.cpp
@@ -4,6 +4,7 @@
 
 #include "Algo/Unique.h"
 #include "BaseGizmos/CombinedTransformGizmo.h"
+#include "BaseGizmos/GizmoBaseComponent.h"
 #include "ContextObjectStore.h"
 #include "Drawing/LineSetComponent.h"
 #include "Drawing/MeshElementsVisualizer.h"
@@ -426,6 +427,19 @@ void UUVSelectTool::Setup()
 	TransformGizmo->SetActiveTarget(TransformProxy, GetToolManager());
 	TransformGizmo->SetVisibility(ViewportButtonsAPI->GetGizmoMode() != UUVToolViewportButtonsAPI::EGizmoMode::Select);
 
+	// Tell the gizmo to be drawn on top even over translucent-mode materials.
+	// Note: this may someday not be necessary, if we get this to work properly by default. Normally we can't
+	// use this approach in modeling mode because it adds dithering to the occluded sections, but we are able
+	// to disable that in the uv editor viewports.
+	for (UActorComponent* Component : TransformGizmo->GetGizmoActor()->GetComponents())
+	{
+		UGizmoBaseComponent* GizmoComponent = Cast<UGizmoBaseComponent>(Component);
+		if (GizmoComponent)
+		{
+			GizmoComponent->bUseEditorCompositing = true;
+		}
+	}
+
 	LivePreviewGeometryActor = Targets[0]->AppliedPreview->GetWorld()->SpawnActor<APreviewGeometryActor>(
 		FVector::ZeroVector, FRotator(0, 0, 0), FActorSpawnParameters());
 	LivePreviewLineSet = NewObject<ULineSetComponent>(LivePreviewGeometryActor);
diff --git a/Engine/Plugins/Runtime/MeshModelingToolset/Source/ModelingComponents/Private/Selection/MeshSelectionMechanic.cpp b/Engine/Plugins/Runtime/MeshModelingToolset/Source/ModelingComponents/Private/Selection/MeshSelectionMechanic.cpp
--- a/Engine/Plugins/Runtime/MeshModelingToolset/Source/ModelingComponents/Private/Selection/MeshSelectionMechanic.cpp
+++ b/Engine/Plugins/Runtime/MeshModelingToolset/Source/ModelingComponents/Private/Selection/MeshSelectionMechanic.cpp
@@ -296,7 +296,7 @@ void UMeshSelectionMechanic::Setup(UInteractiveTool* ParentToolIn)
 	ClearCurrentSelection();
 
 	TriangleSet = NewObject<UTriangleSetComponent>();
-	TriangleSetMaterial = ToolSetupUtil::GetCustomTwoSidedDepthOffsetMaterial(GetParentTool()->GetToolManager(), TriangleColor, TriangleDepthBias, TriangleOpacity);
+	TriangleSetMaterial = ToolSetupUtil::GetCustomTwoSidedDepthOffsetMaterial(GetParentTool()->GetToolManager(), TriangleColor, DepthBias, TriangleOpacity);
 	
 	LineSet = NewObject<ULineSetComponent>();
 	LineSet->SetLineMaterial(ToolSetupUtil::GetDefaultLineComponentMaterial(
diff --git a/Engine/Plugins/Runtime/MeshModelingToolset/Source/ModelingComponents/Public/Selection/MeshSelectionMechanic.h b/Engine/Plugins/Runtime/MeshModelingToolset/Source/ModelingComponents/Public/Selection/MeshSelectionMechanic.h
--- a/Engine/Plugins/Runtime/MeshModelingToolset/Source/ModelingComponents/Public/Selection/MeshSelectionMechanic.h
+++ b/Engine/Plugins/Runtime/MeshModelingToolset/Source/ModelingComponents/Public/Selection/MeshSelectionMechanic.h
@@ -159,7 +159,6 @@ protected:
 	float LineThickness = 3;
 	float PointThickness = 6;
 	float DepthBias = 0.7;
-	float TriangleDepthBias = -0.1;
 
 	FVector3d CurrentSelectionCentroid;
 	int32 CentroidTimestamp = -1;
diff --git a/Engine/Source/Runtime/Engine/Public/ShowFlags.h b/Engine/Source/Runtime/Engine/Public/ShowFlags.h
--- a/Engine/Source/Runtime/Engine/Public/ShowFlags.h
+++ b/Engine/Source/Runtime/Engine/Public/ShowFlags.h
@@ -375,6 +375,7 @@ private:
 		SetVectorFields(false);
 		SetGBufferHints(false);
 		SetCompositeEditorPrimitives(InitMode == ESFIM_Editor || InitMode == ESFIM_VREditing);
+		SetOpaqueCompositeEditorPrimitives(false);
 		SetTestImage(false);
 		SetVisualizeDOF(false);
 		SetVertexColors(false);
diff --git a/Engine/Source/Runtime/Engine/Public/ShowFlagsValues.inl b/Engine/Source/Runtime/Engine/Public/ShowFlagsValues.inl
--- a/Engine/Source/Runtime/Engine/Public/ShowFlagsValues.inl
+++ b/Engine/Source/Runtime/Engine/Public/ShowFlagsValues.inl
@@ -86,6 +86,8 @@ SHOWFLAG_FIXED_IN_SHIPPING(0, GBufferHints, SFG_Developer, NSLOCTEXT("UnrealEd",
 SHOWFLAG_ALWAYS_ACCESSIBLE(MotionBlur, SFG_PostProcess, NSLOCTEXT("UnrealEd", "MotionBlurSF", "Motion Blur"))
 /** Whether to render the editor gizmos and other foreground editor widgets off screen and apply them after post process, only needed for the editor */
 SHOWFLAG_FIXED_IN_SHIPPING(0, CompositeEditorPrimitives, SFG_Developer, NSLOCTEXT("UnrealEd", "CompositeEditorPrimitivesSF", "Composite Editor Primitives"))
+/** When CompositeEditorPrimitives is set, determines whether occluded portions are drawn opaquely on top, or dithered when they are occluded */
+SHOWFLAG_FIXED_IN_SHIPPING(0, OpaqueCompositeEditorPrimitives, SFG_Developer, NSLOCTEXT("UnrealEd", "OpaqueCompositeEditorPrimitives", "Make Composite Editor Primitives Opaque"))
 /** Shows a test image that allows to tweak the monitor colors, borders and allows to judge image and temporal aliasing  */
 SHOWFLAG_FIXED_IN_SHIPPING(0, TestImage, SFG_Developer, NSLOCTEXT("UnrealEd", "TestImageSF", "Test Image"))
 /** Helper to tweak depth of field */
diff --git a/Engine/Source/Runtime/Renderer/Private/PostProcess/PostProcessCompositeEditorPrimitives.cpp b/Engine/Source/Runtime/Renderer/Private/PostProcess/PostProcessCompositeEditorPrimitives.cpp
--- a/Engine/Source/Runtime/Renderer/Private/PostProcess/PostProcessCompositeEditorPrimitives.cpp
+++ b/Engine/Source/Runtime/Renderer/Private/PostProcess/PostProcessCompositeEditorPrimitives.cpp
@@ -13,12 +13,6 @@
 
 namespace
 {
-TAutoConsoleVariable<float> CVarEditorOpaqueGizmo(
-	TEXT("r.Editor.OpaqueGizmo"),
-	0.0f,
-	TEXT("0..1\n0: occluded gizmo is partly transparent (default), 1:gizmo is never occluded"),
-	ECVF_RenderThreadSafe);
-
 TAutoConsoleVariable<int32> CVarEditorTemporalUpsampleDepth(
 	TEXT("r.Editor.TemporalUpsampleDepth"), 2,
 	TEXT("Temporal upsample factor of the depth buffer for depth testing editor primitives against."),
@@ -534,7 +528,7 @@ FScreenPassTexture AddEditorPrimitivePass(
 	{
 		FRHISamplerState* PointClampSampler = TStaticSamplerState<SF_Point, AM_Clamp, AM_Clamp, AM_Clamp>::GetRHI();
 
-		const bool bOpaqueEditorGizmo = CVarEditorOpaqueGizmo.GetValueOnRenderThread() != 0 || View.Family->EngineShowFlags.Wireframe;
+		const bool bOpaqueEditorGizmo = View.Family->EngineShowFlags.OpaqueCompositeEditorPrimitives || View.Family->EngineShowFlags.Wireframe;
 
 		FCompositeEditorPrimitivesPS::FParameters* PassParameters = GraphBuilder.AllocParameters<FCompositeEditorPrimitivesPS::FParameters>();
 		PassParameters->RenderTargets[0] = Output.GetRenderTargetBinding();
