UVEditor: Add view menu to the 3D viewport, allowing control over the current rendering mode like the main level editor and static mesh editor provide.

#rnx

diff --git a/Engine/Plugins/Experimental/UVEditor/Source/UVEditor/Private/SUVEditor3DViewport.cpp b/Engine/Plugins/Experimental/UVEditor/Source/UVEditor/Private/SUVEditor3DViewport.cpp
--- a/Engine/Plugins/Experimental/UVEditor/Source/UVEditor/Private/SUVEditor3DViewport.cpp
+++ b/Engine/Plugins/Experimental/UVEditor/Source/UVEditor/Private/SUVEditor3DViewport.cpp
@@ -33,9 +33,8 @@ void SUVEditor3DViewport::BindCommands()
 
 TSharedPtr<SWidget> SUVEditor3DViewport::MakeViewportToolbar()
 {
-	return SNew(SUVEditor3DViewportToolBar)
+	return SNew(SUVEditor3DViewportToolBar, SharedThis(this))
 		.CommandList(CommandList);
 }
 
-
 #undef LOCTEXT_NAMESPACE
diff --git a/Engine/Plugins/Experimental/UVEditor/Source/UVEditor/Private/SUVEditor3DViewportToolBar.cpp b/Engine/Plugins/Experimental/UVEditor/Source/UVEditor/Private/SUVEditor3DViewportToolBar.cpp
--- a/Engine/Plugins/Experimental/UVEditor/Source/UVEditor/Private/SUVEditor3DViewportToolBar.cpp
+++ b/Engine/Plugins/Experimental/UVEditor/Source/UVEditor/Private/SUVEditor3DViewportToolBar.cpp
@@ -5,14 +5,21 @@
 #include "SUVEditor3DViewportToolBar.h"
 
 #include "UVEditorCommands.h"
+#include "EditorViewportCommands.h"
 #include "EditorStyleSet.h"
 #include "Framework/MultiBox/MultiBoxBuilder.h"
 #include "UVEditorStyle.h"
 #include "Widgets/Layout/SBorder.h"
 #include "Widgets/SBoxPanel.h"
+#include "Editor/UnrealEd/Public/SEditorViewportToolBarMenu.h"
+#include "SEditorViewportViewMenu.h"
+#include "SUVEditor3DViewport.h"
 
-void SUVEditor3DViewportToolBar::Construct(const FArguments& InArgs)
+#define LOCTEXT_NAMESPACE "SUVEditor3DViewportToolBar"
+
+void SUVEditor3DViewportToolBar::Construct(const FArguments& InArgs, TSharedPtr<class SUVEditor3DViewport> InUVEditor3DViewport)
 {
+	UVEditor3DViewportPtr = InUVEditor3DViewport;
 	CommandList = InArgs._CommandList;
 
 	const FMargin ToolbarSlotPadding(4.0f, 1.0f);
@@ -28,6 +35,13 @@ void SUVEditor3DViewportToolBar::Construct(const FArguments& InArgs)
 		]
 	];
 
+	MainBoxPtr->AddSlot()
+		.Padding(ToolbarSlotPadding)
+		.HAlign(HAlign_Left)
+		[
+			MakeDisplayToolBar(InArgs._Extenders)
+		];
+
 	MainBoxPtr->AddSlot()
 		.Padding(ToolbarSlotPadding)
 		.HAlign(HAlign_Right)
@@ -38,6 +52,15 @@ void SUVEditor3DViewportToolBar::Construct(const FArguments& InArgs)
 	SViewportToolBar::Construct(SViewportToolBar::FArguments());
 }
 
+TSharedRef<SWidget> SUVEditor3DViewportToolBar::MakeDisplayToolBar(const TSharedPtr<FExtender> InExtenders)
+{
+	TSharedRef<SEditorViewport> ViewportRef = StaticCastSharedPtr<SEditorViewport>(UVEditor3DViewportPtr.Pin()).ToSharedRef();
+
+	return SNew(SEditorViewportViewMenu, ViewportRef, SharedThis(this))
+		.Cursor(EMouseCursor::Default)
+		.MenuExtenders(InExtenders);
+}
+
 TSharedRef<SWidget> SUVEditor3DViewportToolBar::MakeToolBar(const TSharedPtr<FExtender> InExtenders)
 {
 	// The following is modeled after portions of STransformViewportToolBar, which gets 
@@ -76,4 +99,6 @@ TSharedRef<SWidget> SUVEditor3DViewportToolBar::MakeToolBar(const TSharedPtr<FEx
 	ToolbarBuilder.EndSection();
 
 	return ToolbarBuilder.MakeWidget();
-}
\ No newline at end of file
+}
+
+#undef LOCTEXT_NAMESPACE
\ No newline at end of file
diff --git a/Engine/Plugins/Experimental/UVEditor/Source/UVEditor/Public/SUVEditor3DViewportToolBar.h b/Engine/Plugins/Experimental/UVEditor/Source/UVEditor/Public/SUVEditor3DViewportToolBar.h
--- a/Engine/Plugins/Experimental/UVEditor/Source/UVEditor/Public/SUVEditor3DViewportToolBar.h
+++ b/Engine/Plugins/Experimental/UVEditor/Source/UVEditor/Public/SUVEditor3DViewportToolBar.h
@@ -21,10 +21,14 @@ public:
 		SLATE_ARGUMENT(TSharedPtr<FExtender>, Extenders)
 	SLATE_END_ARGS()
 
-	void Construct(const FArguments& InArgs);
+	void Construct(const FArguments& InArgs, TSharedPtr<class SUVEditor3DViewport> InUVEditor3DViewport);
 
 private:
+	TSharedRef<SWidget> MakeDisplayToolBar(const TSharedPtr<FExtender> InExtenders);
 	TSharedRef<SWidget> MakeToolBar(const TSharedPtr<FExtender> InExtenders);
 
+	/** The viewport that we are in */
+	TWeakPtr<class SUVEditor3DViewport> UVEditor3DViewportPtr;
+
 	TSharedPtr<FUICommandList> CommandList;
 };
