UVEditor: Adjust the handling of viewports near and far clipping planes to deal with z-fighting precision issues and closer viewing of 3D asset surfaces.

#rnx

diff --git a/Engine/Plugins/Experimental/UVEditor/Source/UVEditor/Private/UVEditor2DViewportBehaviorTargets.cpp b/Engine/Plugins/Experimental/UVEditor/Source/UVEditor/Private/UVEditor2DViewportBehaviorTargets.cpp
--- a/Engine/Plugins/Experimental/UVEditor/Source/UVEditor/Private/UVEditor2DViewportBehaviorTargets.cpp
+++ b/Engine/Plugins/Experimental/UVEditor/Source/UVEditor/Private/UVEditor2DViewportBehaviorTargets.cpp
@@ -90,6 +90,9 @@ void FUVEditor2DMouseWheelZoomBehaviorTarget::OnMouseWheelScrollUp(const FInputD
 
 	FVector NewLocation = OriginalLocation + (ZoomInProportion * TToPlane * CurrentPos.WorldRay.Direction);
 
+	ViewportClient->OverrideFarClipPlane(NewLocation.Z - CameraFarPlaneWorldZ);
+	ViewportClient->OverrideNearClipPlane(NewLocation.Z * (1.0-CameraNearPlaneProportionZ));
+
 	// Don't zoom in so far that the XY plane lies in front of our near clipping plane, or else everything
 	// will suddenly disappear.
 	if (NewLocation.Z > ViewportClient->GetNearClipPlane())
@@ -102,7 +105,12 @@ void FUVEditor2DMouseWheelZoomBehaviorTarget::OnMouseWheelScrollDown(const FInpu
 {
 	FVector OriginalLocation = ViewportClient->GetViewLocation();
 	double TToPlane = -OriginalLocation.Z / CurrentPos.WorldRay.Direction.Z;
-	ViewportClient->SetViewLocation(OriginalLocation - (ZoomOutProportion * TToPlane * CurrentPos.WorldRay.Direction));
+	FVector NewLocation = OriginalLocation - (ZoomOutProportion * TToPlane * CurrentPos.WorldRay.Direction);
+
+	ViewportClient->OverrideFarClipPlane(NewLocation.Z - CameraFarPlaneWorldZ);
+	ViewportClient->OverrideNearClipPlane(NewLocation.Z * (1.0 - CameraNearPlaneProportionZ));
+
+	ViewportClient->SetViewLocation(NewLocation);
 }
 
 void FUVEditor2DMouseWheelZoomBehaviorTarget::SetZoomAmount(double PercentZoomIn)
@@ -114,4 +122,14 @@ void FUVEditor2DMouseWheelZoomBehaviorTarget::SetZoomAmount(double PercentZoomIn
 	// Set the zoom out proportion such that (1 + ZoomOutProportion)(1 - ZoomInProportion) = 1
 	// so that zooming in and then zooming out will return to the same zoom level.
 	ZoomOutProportion = ZoomInProportion / (1 - ZoomInProportion);
+}
+
+void FUVEditor2DMouseWheelZoomBehaviorTarget::SetCameraFarPlaneWorldZ(double CameraFarPlaneWorldZIn)
+{
+	CameraFarPlaneWorldZ = CameraFarPlaneWorldZIn;
+}
+
+void FUVEditor2DMouseWheelZoomBehaviorTarget::SetCameraNearPlaneProportionZ(double CameraNearPlaneProportionZIn)
+{
+	CameraNearPlaneProportionZ = CameraNearPlaneProportionZIn;
 }
\ No newline at end of file
diff --git a/Engine/Plugins/Experimental/UVEditor/Source/UVEditor/Private/UVEditor2DViewportClient.cpp b/Engine/Plugins/Experimental/UVEditor/Source/UVEditor/Private/UVEditor2DViewportClient.cpp
--- a/Engine/Plugins/Experimental/UVEditor/Source/UVEditor/Private/UVEditor2DViewportClient.cpp
+++ b/Engine/Plugins/Experimental/UVEditor/Source/UVEditor/Private/UVEditor2DViewportClient.cpp
@@ -39,9 +39,11 @@ FUVEditor2DViewportClient::FUVEditor2DViewportClient(FEditorModeTools* InModeToo
 	BehaviorSet->Add(ScrollBehavior);
 
 	ZoomBehaviorTarget = MakeUnique<FUVEditor2DMouseWheelZoomBehaviorTarget>(this);
+	ZoomBehaviorTarget->SetCameraFarPlaneWorldZ(CameraFarPlaneWorldZ);
+	ZoomBehaviorTarget->SetCameraNearPlaneProportionZ(CameraNearPlaneProportionZ);
 	UMouseWheelInputBehavior* ZoomBehavior = NewObject<UMouseWheelInputBehavior>();
 	ZoomBehavior->Initialize(ZoomBehaviorTarget.Get());
-	ZoomBehavior->SetDefaultPriority(DEFAULT_VIEWPORT_BEHAVIOR_PRIORITY);
+	ZoomBehavior->SetDefaultPriority(DEFAULT_VIEWPORT_BEHAVIOR_PRIORITY);	
 	BehaviorSet->Add(ZoomBehavior);
 
 	ModeTools->GetInteractiveToolsContext()->InputRouter->RegisterSource(this);
@@ -55,6 +57,7 @@ FUVEditor2DViewportClient::FUVEditor2DViewportClient(FEditorModeTools* InModeToo
 	}
 }
 
+
 const UInputBehaviorSet* FUVEditor2DViewportClient::GetInputBehaviors() const
 {
 	return BehaviorSet;
diff --git a/Engine/Plugins/Experimental/UVEditor/Source/UVEditor/Private/UVEditor3DViewportClient.cpp b/Engine/Plugins/Experimental/UVEditor/Source/UVEditor/Private/UVEditor3DViewportClient.cpp
new file mode 100644
--- /dev/null
+++ b/Engine/Plugins/Experimental/UVEditor/Source/UVEditor/Private/UVEditor3DViewportClient.cpp
@@ -0,0 +1,20 @@
+// Copyright Epic Games, Inc. All Rights Reserved.
+
+#include "UVEditor3DViewportClient.h"
+
+#include "BaseBehaviors/ClickDragBehavior.h"
+#include "BaseBehaviors/MouseWheelBehavior.h"
+#include "ContextObjectStore.h"
+#include "EditorModeManager.h"
+#include "EdModeInteractiveToolsContext.h"
+#include "Drawing/MeshDebugDrawing.h"
+#include "FrameTypes.h"
+#include "UVEditorMode.h"
+
+FUVEditor3DViewportClient::FUVEditor3DViewportClient(FEditorModeTools* InModeTools,
+	FPreviewScene* InPreviewScene, const TWeakPtr<SEditorViewport>& InEditorViewportWidget)
+	: FEditorViewportClient(InModeTools, InPreviewScene, InEditorViewportWidget)
+{
+	// We want our near clip plane to be quite close so that we can zoom in further.
+	OverrideNearClipPlane(KINDA_SMALL_NUMBER);
+}
diff --git a/Engine/Plugins/Experimental/UVEditor/Source/UVEditor/Public/UVEditor2DViewportBehaviorTargets.h b/Engine/Plugins/Experimental/UVEditor/Source/UVEditor/Public/UVEditor2DViewportBehaviorTargets.h
--- a/Engine/Plugins/Experimental/UVEditor/Source/UVEditor/Public/UVEditor2DViewportBehaviorTargets.h
+++ b/Engine/Plugins/Experimental/UVEditor/Source/UVEditor/Public/UVEditor2DViewportBehaviorTargets.h
@@ -55,10 +55,15 @@ public:
 	 */
 	virtual void SetZoomAmount(double PercentZoomIn);
 
+	void SetCameraFarPlaneWorldZ(double CameraFarPlaneWorldZIn);
+	void SetCameraNearPlaneProportionZ(double CameraFarPlaneProportionZIn );
+
 	inline static const double DEFAULT_ZOOM_AMOUNT = 20;
 protected:
 	FUVEditor2DViewportClient* ViewportClient = nullptr;
 	double ZoomInProportion;
 	double ZoomOutProportion;
 	double ZoomInLimit;
+	double CameraFarPlaneWorldZ;
+	double CameraNearPlaneProportionZ;
 };
\ No newline at end of file
diff --git a/Engine/Plugins/Experimental/UVEditor/Source/UVEditor/Public/UVEditor2DViewportClient.h b/Engine/Plugins/Experimental/UVEditor/Source/UVEditor/Public/UVEditor2DViewportClient.h
--- a/Engine/Plugins/Experimental/UVEditor/Source/UVEditor/Public/UVEditor2DViewportClient.h
+++ b/Engine/Plugins/Experimental/UVEditor/Source/UVEditor/Public/UVEditor2DViewportClient.h
@@ -34,6 +34,8 @@ public:
 	bool CanSetWidgetMode(UE::Widget::EWidgetMode NewMode) const override;
 	void SetWidgetMode(UE::Widget::EWidgetMode NewMode) override;
 	UE::Widget::EWidgetMode GetWidgetMode() const override;
+	// Overriding base class visibility
+	using FEditorViewportClient::OverrideNearClipPlane;
 
 	// IInputBehaviorSource
 	virtual const UInputBehaviorSet* GetInputBehaviors() const override;
@@ -46,6 +48,12 @@ protected:
 	UInputBehaviorSet* BehaviorSet;
 	UUVToolViewportButtonsAPI* ViewportButtonsAPI;
 
+	// Position to place the camera far plane relative to world z
+	float CameraFarPlaneWorldZ = -10.0;
+	// The near plane gets positioned some proportion to z = 0. We don't use a constant value because our depth offset values are percentage-based
+	// Lower proportions move the plane nearer to world z
+	float CameraNearPlaneProportionZ = 0.8;
+
 	// Note that it's generally less hassle if the unique ptr types are complete here,
 	// not forward declared, else we get compile errors if their destruction shows up
 	// anywhere in the header.
diff --git a/Engine/Plugins/Experimental/UVEditor/Source/UVEditor/Public/UVEditor3DViewportClient.h b/Engine/Plugins/Experimental/UVEditor/Source/UVEditor/Public/UVEditor3DViewportClient.h
--- a/Engine/Plugins/Experimental/UVEditor/Source/UVEditor/Public/UVEditor3DViewportClient.h
+++ b/Engine/Plugins/Experimental/UVEditor/Source/UVEditor/Public/UVEditor3DViewportClient.h
@@ -19,8 +19,12 @@ enum EUVEditor3DViewportClientCameraMode {
  */
 class UVEDITOR_API FUVEditor3DViewportClient : public FEditorViewportClient
 {
-	// Inherit all the same constructors
-	using FEditorViewportClient::FEditorViewportClient;
+public:
+
+	FUVEditor3DViewportClient(FEditorModeTools* InModeTools, FPreviewScene* InPreviewScene = nullptr,
+		const TWeakPtr<SEditorViewport>& InEditorViewportWidget = nullptr);
+
+	virtual ~FUVEditor3DViewportClient() {}
 
 	// FEditorViewportClient
 	virtual bool ShouldOrbitCamera() const override {
