UVEditor: Fix buttons api not respecting "broadcast" argument. Fix undo/redo in select tool being able to create a mismatch between current selection and selection mode.

#rnx

diff --git a/Engine/Plugins/Experimental/UVEditor/Source/UVEditorTools/Private/UVSelectTool.cpp b/Engine/Plugins/Experimental/UVEditor/Source/UVEditorTools/Private/UVSelectTool.cpp
--- a/Engine/Plugins/Experimental/UVEditor/Source/UVEditorTools/Private/UVSelectTool.cpp
+++ b/Engine/Plugins/Experimental/UVEditor/Source/UVEditorTools/Private/UVSelectTool.cpp
@@ -510,7 +510,47 @@ void UUVSelectTool::Shutdown(EToolShutdownType ShutdownType)
 
 void UUVSelectTool::SetSelection(const FDynamicMeshSelection& NewSelection, bool bBroadcastOnSelectionChanged)
 {
-	SelectionMechanic->SetSelection(NewSelection, bBroadcastOnSelectionChanged, false);
+	SelectionMechanic->SetSelection(NewSelection, bBroadcastOnSelectionChanged, 
+		false); // Don't emit undo because this function is called from undo
+	
+	// Make sure the current selection mode is compatible with the new selection we received. Don't broadcast
+	// this part because presumably we've already responded to selection change if bBroadcastOnSelectionChanged
+	// was true above.
+	// TODO: there are a couple things that are not ideal about the below. One is that we always change to
+	// triangle mode when we don't know if the triangles came from island or mesh selection mode. Another is
+	// that we change the selection mode in the mechanic directly rather than going through ChangeSelectionMode,
+	// since we don't want to do the conversions/broadcasts that the setter performs. Still, it's not worth
+	// improving this further because the proper solution will probably involve transacting the selection mode
+	// changes, which we'll probably implement while moving selection up to mode level (along with other changes
+	// that would probably stomp anything we do here)
+	UUVToolViewportButtonsAPI::ESelectionMode CurrentMode = ViewportButtonsAPI->GetSelectionMode();
+	switch (NewSelection.Type)
+	{
+	case FDynamicMeshSelection::EType::Vertex:
+		if (CurrentMode != UUVToolViewportButtonsAPI::ESelectionMode::Vertex)
+		{
+			ViewportButtonsAPI->SetSelectionMode(UUVToolViewportButtonsAPI::ESelectionMode::Vertex, false);
+			SelectionMechanic->SelectionMode = EMeshSelectionMechanicMode::Vertex;
+		}
+		break;
+	case FDynamicMeshSelection::EType::Edge:
+		if (CurrentMode != UUVToolViewportButtonsAPI::ESelectionMode::Edge)
+		{
+			ViewportButtonsAPI->SetSelectionMode(UUVToolViewportButtonsAPI::ESelectionMode::Edge, false);
+			SelectionMechanic->SelectionMode = EMeshSelectionMechanicMode::Edge;
+		}
+		break;
+	case FDynamicMeshSelection::EType::Triangle:
+		if (CurrentMode != UUVToolViewportButtonsAPI::ESelectionMode::Triangle
+			&& CurrentMode != UUVToolViewportButtonsAPI::ESelectionMode::Island
+			&& CurrentMode != UUVToolViewportButtonsAPI::ESelectionMode::Mesh)
+		{
+			ViewportButtonsAPI->SetSelectionMode(UUVToolViewportButtonsAPI::ESelectionMode::Triangle, false);
+			SelectionMechanic->SelectionMode = EMeshSelectionMechanicMode::Triangle;
+		}
+		break;
+	}
+
 }
 
 void UUVSelectTool::SetGizmoTransform(const FTransform& NewTransform)
@@ -565,7 +605,7 @@ void UUVSelectTool::UpdateSelectionMode()
 		TargetMode = EMeshSelectionMechanicMode::Vertex;
 		break;
 	}
-	SelectionMechanic->ChangeSelectionMode(TargetMode);
+	SelectionMechanic->ChangeSelectionMode(TargetMode); // broadcast and emit undo if needed
 }
 
 void UUVSelectTool::OnSelectionChanged()
diff --git a/Engine/Plugins/Experimental/UVEditor/Source/UVEditorTools/Public/UVToolContextObjects.h b/Engine/Plugins/Experimental/UVEditor/Source/UVEditorTools/Public/UVToolContextObjects.h
--- a/Engine/Plugins/Experimental/UVEditor/Source/UVEditorTools/Public/UVToolContextObjects.h
+++ b/Engine/Plugins/Experimental/UVEditor/Source/UVEditorTools/Public/UVToolContextObjects.h
@@ -155,7 +155,10 @@ public:
 	void SetGizmoMode(EGizmoMode ModeIn, bool bBroadcast = true)
 	{
 		GizmoMode = ModeIn;
-		OnGizmoModeChange.Broadcast(GizmoMode);
+		if (bBroadcast)
+		{
+			OnGizmoModeChange.Broadcast(GizmoMode);
+		}
 	}
 
 	EGizmoMode GetGizmoMode()
@@ -177,7 +180,10 @@ public:
 	void SetSelectionMode(ESelectionMode ModeIn, bool bBroadcast = true)
 	{
 		SelectionMode = ModeIn;
-		OnSelectionModeChange.Broadcast(SelectionMode);
+		if (bBroadcast)
+		{
+			OnSelectionModeChange.Broadcast(SelectionMode);
+		}
 	}
 
 	ESelectionMode GetSelectionMode()
