UVEditor: Add ability to display status bar messages and error messages to UV editor tools. Also add tool icon to accept/cancel overlay.

#rnx

diff --git a/Engine/Plugins/Experimental/UVEditor/Source/UVEditor/Private/UVEditorMode.cpp b/Engine/Plugins/Experimental/UVEditor/Source/UVEditor/Private/UVEditorMode.cpp
--- a/Engine/Plugins/Experimental/UVEditor/Source/UVEditor/Private/UVEditorMode.cpp
+++ b/Engine/Plugins/Experimental/UVEditor/Source/UVEditor/Private/UVEditorMode.cpp
@@ -47,9 +47,6 @@ namespace UVEditorModeLocals
 	// The layer we open when we first open the UV editor
 	const int32 DefaultUVLayerIndex = 0;
 
-	// Determines what the default tool is
-	const FString DefaultToolIdentifier = TEXT("UVSelectTool");
-
 	const FText UVLayerChangeTransactionName = LOCTEXT("UVLayerChangeTransactionName", "Change UV Layer");
 
 	/** 
@@ -167,29 +164,34 @@ void UUVEditorMode::RegisterTools()
 {
 	const FUVEditorCommands& CommandInfos = FUVEditorCommands::Get();
 
+	// Note that the identifiers below need to match the command names so that the tool icons can 
+	// be easily retrieved from the active tool name in UVEditorModeToolkit::OnToolStarted. Otherwise
+	// we would need to keep some other mapping from tool identifier to tool icon.
+
 	UUVSelectToolBuilder* UVSelectToolBuilder = NewObject<UUVSelectToolBuilder>();
 	UVSelectToolBuilder->Targets = &ToolInputObjects;
-	RegisterTool(CommandInfos.BeginSelectTool, TEXT("UVSelectTool"), UVSelectToolBuilder);
+	DefaultToolIdentifier = TEXT("BeginSelectTool");
+	RegisterTool(CommandInfos.BeginSelectTool, DefaultToolIdentifier, UVSelectToolBuilder);
 
 	UUVEditorLayoutToolBuilder* UVEditorLayoutToolBuilder = NewObject<UUVEditorLayoutToolBuilder>();
 	UVEditorLayoutToolBuilder->Targets = &ToolInputObjects;
-	RegisterTool(CommandInfos.BeginLayoutTool, TEXT("UVLayoutTool"), UVEditorLayoutToolBuilder);
+	RegisterTool(CommandInfos.BeginLayoutTool, TEXT("BeginLayoutTool"), UVEditorLayoutToolBuilder);
 
 	UUVEditorParameterizeMeshToolBuilder* UVEditorParameterizeMeshToolBuilder = NewObject<UUVEditorParameterizeMeshToolBuilder>();
 	UVEditorParameterizeMeshToolBuilder->Targets = &ToolInputObjects;
-	RegisterTool(CommandInfos.BeginParameterizeMeshTool, TEXT("UVParameterizeMeshTool"), UVEditorParameterizeMeshToolBuilder);
+	RegisterTool(CommandInfos.BeginParameterizeMeshTool, TEXT("BeginParameterizeMeshTool"), UVEditorParameterizeMeshToolBuilder);
 
 	UUVEditorChannelEditToolBuilder* UVEditorChannelEditToolBuilder = NewObject<UUVEditorChannelEditToolBuilder>();
 	UVEditorChannelEditToolBuilder->Targets = &ToolInputObjects;
-	RegisterTool(CommandInfos.BeginChannelEditTool, TEXT("UVChannelEditTool"), UVEditorChannelEditToolBuilder);
+	RegisterTool(CommandInfos.BeginChannelEditTool, TEXT("BeginChannelEditTool"), UVEditorChannelEditToolBuilder);
 
 	UUVEditorSeamToolBuilder* UVEditorSeamToolBuilder = NewObject<UUVEditorSeamToolBuilder>();
 	UVEditorSeamToolBuilder->Targets = &ToolInputObjects;
-	RegisterTool(CommandInfos.BeginSeamTool, TEXT("UVSeamTool"), UVEditorSeamToolBuilder);
+	RegisterTool(CommandInfos.BeginSeamTool, TEXT("BeginSeamTool"), UVEditorSeamToolBuilder);
 
 	UUVEditorRecomputeUVsToolBuilder* UVEditorRecomputeUVsToolBuilder = NewObject<UUVEditorRecomputeUVsToolBuilder>();
 	UVEditorRecomputeUVsToolBuilder->Targets = &ToolInputObjects;
-	RegisterTool(CommandInfos.BeginRecomputeUVsTool, TEXT("UVRecomputeUVsTool"), UVEditorRecomputeUVsToolBuilder);
+	RegisterTool(CommandInfos.BeginRecomputeUVsTool, TEXT("BeginRecomputeUVsTool"), UVEditorRecomputeUVsToolBuilder);
 }
 
 void UUVEditorMode::CreateToolkit()
@@ -208,12 +210,12 @@ UObject* UUVEditorMode::GetBackgroundSettingsObject()
 
 void UUVEditorMode::ActivateDefaultTool()
 {
-	GetInteractiveToolsContext()->StartTool(UVEditorModeLocals::DefaultToolIdentifier);
+	GetInteractiveToolsContext()->StartTool(DefaultToolIdentifier);
 }
 
 bool UUVEditorMode::IsDefaultToolActive()
 {
-	return GetInteractiveToolsContext()->IsToolActive(EToolSide::Mouse, UVEditorModeLocals::DefaultToolIdentifier);
+	return GetInteractiveToolsContext()->IsToolActive(EToolSide::Mouse, DefaultToolIdentifier);
 }
 
 void UUVEditorMode::BindCommands()
diff --git a/Engine/Plugins/Experimental/UVEditor/Source/UVEditor/Private/UVEditorModeToolkit.cpp b/Engine/Plugins/Experimental/UVEditor/Source/UVEditor/Private/UVEditorModeToolkit.cpp
--- a/Engine/Plugins/Experimental/UVEditor/Source/UVEditor/Private/UVEditorModeToolkit.cpp
+++ b/Engine/Plugins/Experimental/UVEditor/Source/UVEditor/Private/UVEditorModeToolkit.cpp
@@ -19,6 +19,7 @@
 #include "EdModeInteractiveToolsContext.h"
 #include "EditorModeManager.h"
 #include "UVEditorModeUILayer.h"
+#include "UVEditorStyle.h"
 #include "AssetEditorModeManager.h"
 
 #define LOCTEXT_NAMESPACE "FUVEditorModeToolkit"
@@ -73,49 +74,37 @@ FUVEditorModeToolkit::FUVEditorModeToolkit()
 	// makes it easy to guarantee that GetInlineContent() will always
 	// be ready to work.
 
-	SAssignNew(ToolkitWidget, SBorder)
-		.HAlign(HAlign_Fill)
-		.Padding(4)
-		[
-			SNew(SVerticalBox)
-			+ SVerticalBox::Slot()
-			.AutoHeight()
-			.HAlign(HAlign_Left)
-			.Padding(1.f)
-			[
-				SAssignNew(ToolButtonsContainer, SBorder)
-				.BorderImage(FEditorStyle::GetBrush("ToolPanel.GroupBorder"))
-				.Padding(FMargin(4, 2, 0, 0))
-			]
-
-			+ SVerticalBox::Slot()
-				.AutoHeight()
-				.Padding(2, 0, 0, 0)
-			[
-				SNew(SVerticalBox)
-				+ SVerticalBox::Slot()
-				.AutoHeight()
-				[
-					SAssignNew(ToolWarningArea, STextBlock)
-					.AutoWrapText(true)
-					.Font(FCoreStyle::GetDefaultFontStyle("Bold", 9))
-					.ColorAndOpacity(FSlateColor(FLinearColor(0.9f, 0.15f, 0.15f))) //TODO: This probably needs to not be hardcoded
-					.Text(FText::GetEmpty())
-					.Visibility(EVisibility::Collapsed)
-				]
-
-				+ SVerticalBox::Slot()
-				[
-					SAssignNew(ToolDetailsContainer, SBorder)
-					.BorderImage(FEditorStyle::GetBrush("NoBorder"))
-				]
-			]
-
-		];
+	SAssignNew(ToolkitWidget, SVerticalBox)
+	+ SVerticalBox::Slot()
+	.AutoHeight()
+	[
+		SAssignNew(ToolWarningArea, STextBlock)
+		.AutoWrapText(true)
+		.Font(FCoreStyle::GetDefaultFontStyle("Bold", 9))
+		.ColorAndOpacity(FSlateColor(FLinearColor(0.9f, 0.15f, 0.15f))) //TODO: This probably needs to not be hardcoded
+		.Text(FText::GetEmpty())
+		.Visibility(EVisibility::Collapsed)
+	]
+
+	+ SVerticalBox::Slot()
+	[
+		SAssignNew(ToolDetailsContainer, SBorder)
+		.BorderImage(FEditorStyle::GetBrush("NoBorder"))
+	];
 }
 
 FUVEditorModeToolkit::~FUVEditorModeToolkit()
 {
+	UEdMode* Mode = GetScriptableEditorMode().Get();
+	if (ensure(Mode))
+	{
+		UEditorInteractiveToolsContext* Context = Mode->GetInteractiveToolsContext();
+		if (ensure(Context))
+		{
+			Context->OnToolNotificationMessage.RemoveAll(this);
+			Context->OnToolWarningMessage.RemoveAll(this);
+		}
+	}
 }
 
 void FUVEditorModeToolkit::Init(const TSharedPtr<IToolkitHost>& InitToolkitHost, TWeakObjectPtr<UEdMode> InOwningMode)
@@ -132,6 +121,11 @@ void FUVEditorModeToolkit::Init(const TSharedPtr<IToolkitHost>& InitToolkitHost,
 	// of tool start transactions. See FUVEditorModeToolkit::OnToolStarted for how we issue the transactions.
 	InteractiveToolsContext->ToolManager->ConfigureChangeTrackingMode(EToolChangeTrackingMode::NoChangeTracking);
 
+	// Set up tool message areas
+	ClearNotification();
+	ClearWarning();
+	GetScriptableEditorMode()->GetInteractiveToolsContext()->OnToolNotificationMessage.AddSP(this, &FUVEditorModeToolkit::PostNotification);
+	GetScriptableEditorMode()->GetInteractiveToolsContext()->OnToolWarningMessage.AddSP(this, &FUVEditorModeToolkit::PostWarning);
 
 	// Hook up the tool detail panel
 	ToolDetailsContainer->SetContent(DetailsView.ToSharedRef());
@@ -151,15 +145,14 @@ void FUVEditorModeToolkit::Init(const TSharedPtr<IToolkitHost>& InitToolkitHost,
 		[
 			SNew(SHorizontalBox)
 
-			// TODO: Need to add a tool icon here the way we do in modeling mode.
-			//+SHorizontalBox::Slot()
-			//.AutoWidth()
-			//.VAlign(VAlign_Center)
-			//.Padding(FMargin(0.f, 0.f, 8.f, 0.f))
-			//[
-			//	SNew(SImage)
-			//	.Image_Lambda([this] () { return ActiveToolIcon; })
-			//]
+			+SHorizontalBox::Slot()
+			.AutoWidth()
+			.VAlign(VAlign_Center)
+			.Padding(FMargin(0.f, 0.f, 8.f, 0.f))
+			[
+				SNew(SImage)
+				.Image_Lambda([this] () { return ActiveToolIcon; })
+			]
 
 			+SHorizontalBox::Slot()
 			.AutoWidth()
@@ -348,6 +341,11 @@ void FUVEditorModeToolkit::OnToolStarted(UInteractiveToolManager* Manager, UInte
 
 	ActiveToolName = Tool->GetToolInfo().ToolDisplayName;
 
+	FString ActiveToolIdentifier = GetScriptableEditorMode()->GetToolManager()->GetActiveToolName(EToolSide::Mouse);
+	ActiveToolIdentifier.InsertAt(0, ".");
+	FName ActiveToolIconName = ISlateStyle::Join(FUVEditorCommands::Get().GetContextName(), TCHAR_TO_ANSI(*ActiveToolIdentifier));
+	ActiveToolIcon = FUVEditorStyle::Get().GetOptionalBrush(ActiveToolIconName);
+
 	UUVEditorMode* Mode = Cast<UUVEditorMode>(GetScriptableEditorMode());
 	if (!Mode->IsDefaultToolActive())
 	{
@@ -367,6 +365,8 @@ void FUVEditorModeToolkit::OnToolEnded(UInteractiveToolManager* Manager, UIntera
 	FModeToolkit::OnToolEnded(Manager, Tool);
 
 	ActiveToolName = FText::GetEmpty();
+	ClearNotification();
+	ClearWarning();
 
 	if (IsHosted())
 	{
@@ -412,4 +412,45 @@ void FUVEditorModeToolkit::BuildToolPalette(FName PaletteIndex, class FToolBarBu
 	}
 }
 
+void FUVEditorModeToolkit::PostNotification(const FText& Message)
+{
+	ClearNotification();
+
+	if (ModeUILayer.IsValid())
+	{
+		TSharedPtr<FAssetEditorModeUILayer> ModeUILayerPtr = ModeUILayer.Pin();
+		ActiveToolMessageHandle = GEditor->GetEditorSubsystem<UStatusBarSubsystem>()->PushStatusBarMessage(ModeUILayerPtr->GetStatusBarName(), Message);
+	}
+}
+
+void FUVEditorModeToolkit::ClearNotification()
+{
+	if (ModeUILayer.IsValid())
+	{
+		TSharedPtr<FAssetEditorModeUILayer> ModeUILayerPtr = ModeUILayer.Pin();
+		GEditor->GetEditorSubsystem<UStatusBarSubsystem>()->PopStatusBarMessage(ModeUILayerPtr->GetStatusBarName(), ActiveToolMessageHandle);
+	}
+	ActiveToolMessageHandle.Reset();
+}
+
+void FUVEditorModeToolkit::PostWarning(const FText& Message)
+{
+	if (Message.IsEmpty())
+	{
+		ClearWarning();
+	}
+	else
+	{
+		ToolWarningArea->SetText(Message);
+		ToolWarningArea->SetVisibility(EVisibility::Visible);
+	}
+}
+
+void FUVEditorModeToolkit::ClearWarning()
+{
+	ToolWarningArea->SetText(FText());
+	ToolWarningArea->SetVisibility(EVisibility::Collapsed);
+}
+
+
 #undef LOCTEXT_NAMESPACE
diff --git a/Engine/Plugins/Experimental/UVEditor/Source/UVEditor/Private/UVEditorToolkit.cpp b/Engine/Plugins/Experimental/UVEditor/Source/UVEditor/Private/UVEditorToolkit.cpp
--- a/Engine/Plugins/Experimental/UVEditor/Source/UVEditor/Private/UVEditorToolkit.cpp
+++ b/Engine/Plugins/Experimental/UVEditor/Source/UVEditor/Private/UVEditorToolkit.cpp
@@ -435,15 +435,10 @@ void FUVEditorToolkit::PostInitAssetEditor()
 		TabManager->TryInvokeTab(ViewportTabID);
 	}
 
-	TSharedPtr<FUVEditorModeToolkit> UVModeToolkit = StaticCastSharedPtr<FUVEditorModeToolkit>(UVMode->GetToolkit().Pin());
-
 	// Add the "Apply Changes" button. It should actually be safe to do this almost
 	// any time, even before that toolbar's registration, but it's easier to put most
 	// things into PostInitAssetEditor().
 
-	// TODO: Use the icon that material editor uses. Also, the space between the button
-	// sections is too large, but the button can't go into the default toolbar section or
-	// else it will show up in other asset editors where the section is used.
 	// TODO: We may consider putting actions like these, which are tied to a mode, into
 	// some list of mode actions, and then letting the mode supply them to the owning
 	// asset editor on enter/exit. Revisit when/if this becomes easier to do.
@@ -458,6 +453,8 @@ void FUVEditorToolkit::PostInitAssetEditor()
 	Section.AddEntry(FToolMenuEntry::InitToolBarButton(FUVEditorCommands::Get().ApplyChanges));
 	
 	// Add the channel selection button.
+	check(UVMode->GetToolkit().Pin());
+	FUVEditorModeToolkit* UVModeToolkit = static_cast<FUVEditorModeToolkit*>(UVMode->GetToolkit().Pin().Get());
 	Section.AddEntry(FToolMenuEntry::InitComboButton(
 		"UVEditorChannelMenu",
 		FUIAction(),
diff --git a/Engine/Plugins/Experimental/UVEditor/Source/UVEditor/Public/UVEditorMode.h b/Engine/Plugins/Experimental/UVEditor/Source/UVEditor/Public/UVEditorMode.h
--- a/Engine/Plugins/Experimental/UVEditor/Source/UVEditor/Public/UVEditorMode.h
+++ b/Engine/Plugins/Experimental/UVEditor/Source/UVEditor/Public/UVEditorMode.h
@@ -221,6 +221,7 @@ protected:
 	TArray<TObjectPtr<UInteractiveToolPropertySet>> PropertyObjectsToTick;
 
 	bool bIsActive = false;
+	FString DefaultToolIdentifier;
 
 	static FDateTime AnalyticsLastStartTimestamp;
 };
diff --git a/Engine/Plugins/Experimental/UVEditor/Source/UVEditor/Public/UVEditorModeToolkit.h b/Engine/Plugins/Experimental/UVEditor/Source/UVEditor/Public/UVEditorModeToolkit.h
--- a/Engine/Plugins/Experimental/UVEditor/Source/UVEditor/Public/UVEditorModeToolkit.h
+++ b/Engine/Plugins/Experimental/UVEditor/Source/UVEditor/Public/UVEditorModeToolkit.h
@@ -6,7 +6,9 @@
 
 #include "InteractiveTool.h"
 #include "Toolkits/BaseToolkit.h"
+#include "StatusBarSubsystem.h"
 
+struct FSlateBrush;
 class SBorder;
 class STextBlock;
 class UInteractiveToolPropertySet;
@@ -30,23 +32,28 @@ public:
 	// FModeToolkit
 	virtual void Init(const TSharedPtr<IToolkitHost>& InitToolkitHost, TWeakObjectPtr<UEdMode> InOwningMode) override;
 	virtual FText GetActiveToolDisplayName() const override { return ActiveToolName; }
-
 	/** Returns the Mode specific tabs in the mode toolbar **/
-	virtual void GetToolPaletteNames(TArray<FName>& InPaletteName) const;
-
+	virtual void GetToolPaletteNames(TArray<FName>& InPaletteName) const override;
     /** Returns human readable name for the specific palette tool category **/
-	virtual FText GetToolPaletteDisplayName(FName PaletteName) const;
-
+	virtual FText GetToolPaletteDisplayName(FName PaletteName) const override;
     /** For a specific tool palette category, construct and fill ToolbarBuilder with the category's tools **/
-	virtual void BuildToolPalette(FName PaletteName, class FToolBarBuilder& ToolbarBuilder);
-
+	virtual void BuildToolPalette(FName PaletteName, class FToolBarBuilder& ToolbarBuilder) override;
 	/** Informs the built in FModeToolkit palette building routines about how we want the toolbars to be setup **/
-	virtual bool HasIntegratedToolPalettes() const { return true; }
-	virtual bool HasExclusiveToolPalettes() const { return false; }
+	virtual bool HasIntegratedToolPalettes() const override { return true; }
+	virtual bool HasExclusiveToolPalettes() const override { return false; }
 
 	// IToolkit
 	virtual FName GetToolkitFName() const override;
 	virtual FText GetBaseToolkitName() const override;
+	virtual TSharedPtr<class SWidget> GetInlineContent() const override { return ToolkitWidget; }
+
+	// set/clear notification message area
+	virtual void PostNotification(const FText& Message);
+	virtual void ClearNotification();
+
+	// set/clear warning message area
+	virtual void PostWarning(const FText& Message);
+	virtual void ClearWarning();
 
 protected:
 	const static TArray<FName> PaletteNames_Standard;
@@ -76,6 +83,8 @@ protected:
 	TSharedPtr<SWidget> ViewportOverlayWidget;
 
 	FText ActiveToolName;
+	const FSlateBrush* ActiveToolIcon;
+	FStatusBarMessageHandle ActiveToolMessageHandle;
 
 	void UpdateActiveToolProperties();
 	void InvalidateCachedDetailPanelState(UObject* ChangedObject);
diff --git a/Engine/Plugins/Experimental/UVEditor/Source/UVEditor/UVEditor.Build.cs b/Engine/Plugins/Experimental/UVEditor/Source/UVEditor/UVEditor.Build.cs
--- a/Engine/Plugins/Experimental/UVEditor/Source/UVEditor/UVEditor.Build.cs
+++ b/Engine/Plugins/Experimental/UVEditor/Source/UVEditor/UVEditor.Build.cs
@@ -47,20 +47,20 @@ public class UVEditor : ModuleRules
 				
 				"AdvancedPreviewScene",
 				"ContentBrowser",
-				"GeometryCore",
 				"DynamicMesh",
 				"EditorFramework", // FEditorModeInfo
 				"EditorInteractiveToolsFramework",
 				"EditorStyle",
 				"EditorSubsystem",
+				"GeometryCore",
 				"LevelEditor",
 				"ModelingComponentsEditorOnly", // Static/skeletal mesh tool targets
 				"ModelingComponents",
+				"StatusBar",
 				"ToolMenus",
+				"ToolWidgets",
 				"UVEditorTools",
 				"UVEditorToolsEditorOnly",
-				"ToolWidgets",
-
 				"WorkspaceMenuStructure",
 				// ... add private dependencies that you statically link with here ...	
 			}
diff --git a/Engine/Plugins/Experimental/UVEditor/Source/UVEditorTools/Private/UVEditorSeamTool.cpp b/Engine/Plugins/Experimental/UVEditor/Source/UVEditorTools/Private/UVEditorSeamTool.cpp
--- a/Engine/Plugins/Experimental/UVEditor/Source/UVEditorTools/Private/UVEditorSeamTool.cpp
+++ b/Engine/Plugins/Experimental/UVEditor/Source/UVEditorTools/Private/UVEditorSeamTool.cpp
@@ -476,6 +476,9 @@ void UUVEditorSeamTool::Setup()
 			ReconstructExistingSeamsVisualization();
 		});
 	}
+
+	GetToolManager()->DisplayMessage(LOCTEXT("OnStartSeamTool", "Click to add seams in 2d/3d view. Click start/end point or hit Enter to confirm."),
+		EToolMessageLevel::UserNotification);
 }
 
 void UUVEditorSeamTool::ReconstructExistingSeamsVisualization()
diff --git a/Engine/Plugins/Experimental/UVEditor/Source/UVEditorTools/Private/UVIslandConformalUnwrapAction.cpp b/Engine/Plugins/Experimental/UVEditor/Source/UVEditorTools/Private/UVIslandConformalUnwrapAction.cpp
--- a/Engine/Plugins/Experimental/UVEditor/Source/UVEditorTools/Private/UVIslandConformalUnwrapAction.cpp
+++ b/Engine/Plugins/Experimental/UVEditor/Source/UVEditorTools/Private/UVIslandConformalUnwrapAction.cpp
@@ -62,6 +62,9 @@ bool UUVIslandConformalUnwrapAction::PreCheckAction()
 	if (!CurrentSelection || CurrentSelection->Type != FDynamicMeshSelection::EType::Triangle
 		|| CurrentSelection->Mesh == nullptr || SelectionTargetIndex == INDEX_NONE)
 	{
+		ParentTool->GetToolManager()->DisplayMessage(
+			LOCTEXT("UnwrapErrorSelectionEmpty", "Cannot perform unwrap. Mesh selection must be island triangles."),
+			EToolMessageLevel::UserWarning);
 		return false;
 	}
 
diff --git a/Engine/Plugins/Experimental/UVEditor/Source/UVEditorTools/Private/UVSeamSewAction.cpp b/Engine/Plugins/Experimental/UVEditor/Source/UVEditorTools/Private/UVSeamSewAction.cpp
--- a/Engine/Plugins/Experimental/UVEditor/Source/UVEditorTools/Private/UVSeamSewAction.cpp
+++ b/Engine/Plugins/Experimental/UVEditor/Source/UVEditorTools/Private/UVSeamSewAction.cpp
@@ -158,19 +158,25 @@ bool UUVSeamSewAction::PreCheckAction()
 {
 	if (SelectionTargetIndex == -1 || !CurrentSelection || CurrentSelection->Mesh == nullptr)
 	{
-		UE_LOG(LogGeometry, Warning, TEXT("Cannot sew UVs. Mesh selection was empty."));
+		ParentTool->GetToolManager()->DisplayMessage(
+			LOCTEXT("SewErrorSelectionEmpty", "Cannot sew UVs. Mesh selection was empty."),
+			EToolMessageLevel::UserWarning);
 		return false;
 	}
 
 	if (CurrentSelection->Type != FDynamicMeshSelection::EType::Edge)
 	{
-		UE_LOG(LogGeometry, Warning, TEXT("Cannot sew UVs. Selection was not an edge."));
+		ParentTool->GetToolManager()->DisplayMessage(
+			LOCTEXT("SewErrorSelectionNotEdge", "Cannot sew UVs. Selection was not an edge."),
+			EToolMessageLevel::UserWarning);
 		return false;
 	}
 
 	if (EdgeSewCandidates.Num() == 0)
 	{
-		UE_LOG(LogGeometry, Warning, TEXT("Cannot sew UVs. No viable sew candidate edges selected."));
+		ParentTool->GetToolManager()->DisplayMessage(
+			LOCTEXT("SewErrorSelectionNotBoundary", "Cannot sew UVs. No viable sew candidate edges selected."),
+			EToolMessageLevel::UserWarning);
 		return false;
 	}
 
diff --git a/Engine/Plugins/Experimental/UVEditor/Source/UVEditorTools/Private/UVSelectTool.cpp b/Engine/Plugins/Experimental/UVEditor/Source/UVEditorTools/Private/UVSelectTool.cpp
--- a/Engine/Plugins/Experimental/UVEditor/Source/UVEditorTools/Private/UVSelectTool.cpp
+++ b/Engine/Plugins/Experimental/UVEditor/Source/UVEditorTools/Private/UVSelectTool.cpp
@@ -457,6 +457,10 @@ void UUVSelectTool::Setup()
 		OnSelectionChanged();
 	}
 	UpdateGizmo();
+
+	GetToolManager()->DisplayMessage(LOCTEXT("SelectToolStatusBarMessage", 
+		"Select elements in the viewport and then use one of the edit action buttons."), 
+		EToolMessageLevel::UserNotification);
 }
 
 void UUVSelectTool::Shutdown(EToolShutdownType ShutdownType)
@@ -614,6 +618,8 @@ void UUVSelectTool::OnSelectionChanged()
 	
 	using namespace UVSelectToolLocals;
 
+	ClearWarning();
+
 	const FDynamicMeshSelection& Selection = SelectionMechanic->GetCurrentSelection();
 
 	GetVidPairsFromSelection(Selection, CurrentSelectionVidPairs);
@@ -751,6 +757,11 @@ void UUVSelectTool::OnSelectionChanged()
 	UpdateGizmo();
 }
 
+void UUVSelectTool::ClearWarning()
+{
+	GetToolManager()->DisplayMessage(FText(), EToolMessageLevel::UserWarning);
+}
+
 void UUVSelectTool::UpdateLivePreviewLines()
 {
 	TRACE_CPUPROFILER_EVENT_SCOPE(UVSelectTool_UpdateLivePreviewLines);
@@ -907,6 +918,7 @@ void UUVSelectTool::OnTick(float DeltaTime)
 
 void UUVSelectTool::RequestAction(ESelectToolAction ActionType)
 {
+	ClearWarning();
 	if (PendingAction == ESelectToolAction::NoAction)
 	{
 		PendingAction = ActionType;
@@ -960,6 +972,9 @@ void UUVSelectTool::ApplySplit()
 
 	if (Selection.IsEmpty() || Selection.Type != FDynamicMeshSelection::EType::Edge)
 	{
+		GetToolManager()->DisplayMessage(
+			LOCTEXT("SplitErrorSelectionEmpty", "Cannot split UV's. Edge selection was empty."),
+			EToolMessageLevel::UserWarning);
 		return;
 	}
 
diff --git a/Engine/Plugins/Experimental/UVEditor/Source/UVEditorTools/Public/UVSelectTool.h b/Engine/Plugins/Experimental/UVEditor/Source/UVEditorTools/Public/UVSelectTool.h
--- a/Engine/Plugins/Experimental/UVEditor/Source/UVEditorTools/Public/UVSelectTool.h
+++ b/Engine/Plugins/Experimental/UVEditor/Source/UVEditorTools/Public/UVSelectTool.h
@@ -125,6 +125,7 @@ public:
 
 protected:
 	virtual void OnSelectionChanged();
+	virtual void ClearWarning();
 
 	// Callbacks we'll receive from the gizmo proxy
 	virtual void GizmoTransformChanged(UTransformProxy* Proxy, FTransform Transform);
