Make UEdMode::GetDefaultToolScope non-virtual again to force modes to use mode-level contexts. To work this way, move the initialization of context objects in the UV editor into the UV editor mode, passing in the needed information from the UV asset editor.

#rnx

diff --git a/Engine/Plugins/Experimental/UVEditor/Source/UVEditor/Private/UVEditor2DViewportClient.cpp b/Engine/Plugins/Experimental/UVEditor/Source/UVEditor/Private/UVEditor2DViewportClient.cpp
--- a/Engine/Plugins/Experimental/UVEditor/Source/UVEditor/Private/UVEditor2DViewportClient.cpp
+++ b/Engine/Plugins/Experimental/UVEditor/Source/UVEditor/Private/UVEditor2DViewportClient.cpp
@@ -12,8 +12,10 @@
 #include "UVEditorMode.h"
 
 FUVEditor2DViewportClient::FUVEditor2DViewportClient(FEditorModeTools* InModeTools,
-	FPreviewScene* InPreviewScene, const TWeakPtr<SEditorViewport>& InEditorViewportWidget)
+	FPreviewScene* InPreviewScene, const TWeakPtr<SEditorViewport>& InEditorViewportWidget,
+	UUVToolViewportButtonsAPI* ViewportButtonsAPIIn)
 	: FEditorViewportClient(InModeTools, InPreviewScene, InEditorViewportWidget)
+	, ViewportButtonsAPI(ViewportButtonsAPIIn)
 {
 	ShowWidget(false);
 
@@ -47,14 +49,6 @@ FUVEditor2DViewportClient::FUVEditor2DViewportClient(FEditorModeTools* InModeToo
 	BehaviorSet->Add(ZoomBehavior);
 
 	ModeTools->GetInteractiveToolsContext()->InputRouter->RegisterSource(this);
-
-	UContextObjectStore* ContextStore = ModeTools->GetInteractiveToolsContext()->ContextObjectStore;
-	ViewportButtonsAPI = ContextStore->FindContext<UUVToolViewportButtonsAPI>();
-	if (!ViewportButtonsAPI)
-	{
-		ViewportButtonsAPI = NewObject<UUVToolViewportButtonsAPI>();
-		ContextStore->AddContextObject(ViewportButtonsAPI);
-	}
 }
 
 
diff --git a/Engine/Plugins/Experimental/UVEditor/Source/UVEditor/Private/UVEditorMode.cpp b/Engine/Plugins/Experimental/UVEditor/Source/UVEditor/Private/UVEditorMode.cpp
--- a/Engine/Plugins/Experimental/UVEditor/Source/UVEditor/Private/UVEditorMode.cpp
+++ b/Engine/Plugins/Experimental/UVEditor/Source/UVEditor/Private/UVEditorMode.cpp
@@ -3,9 +3,11 @@
 #include "UVEditorMode.h"
 
 #include "Algo/AnyOf.h"
+#include "AssetEditorModeManager.h"
 #include "ContextObjectStore.h"
 #include "Drawing/MeshElementsVisualizer.h"
 #include "Editor.h"
+#include "EditorViewportClient.h"
 #include "EdModeInteractiveToolsContext.h" //ToolsContext
 #include "EngineAnalytics.h"
 #include "Framework/Commands/UICommandList.h"
@@ -13,6 +15,7 @@
 #include "MeshOpPreviewHelpers.h" // UMeshOpPreviewWithBackgroundCompute
 #include "ModelingToolTargetUtil.h"
 #include "PreviewMesh.h"
+#include "PreviewScene.h"
 #include "TargetInterfaces/AssetBackedTarget.h"
 #include "TargetInterfaces/MaterialProvider.h"
 #include "TargetInterfaces/MeshDescriptionCommitter.h"
@@ -49,6 +52,22 @@ namespace UVEditorModeLocals
 
 	const FText UVLayerChangeTransactionName = LOCTEXT("UVLayerChangeTransactionName", "Change UV Layer");
 
+	void GetCameraState(const FEditorViewportClient& ViewportClientIn, FViewCameraState& CameraStateOut)
+	{
+		FViewportCameraTransform ViewTransform = ViewportClientIn.GetViewTransform();
+		CameraStateOut.bIsOrthographic = false;
+		CameraStateOut.bIsVR = false;
+		CameraStateOut.Position = ViewTransform.GetLocation();
+		CameraStateOut.HorizontalFOVDegrees = ViewportClientIn.ViewFOV;
+		CameraStateOut.AspectRatio = ViewportClientIn.AspectRatio;
+
+		// if using Orbit camera, the rotation in the ViewTransform is not the current camera rotation, it
+		// is set to a different rotation based on the Orbit. So we have to convert back to camera rotation.
+		FRotator ViewRotation = (ViewportClientIn.bUsingOrbitCamera) ?
+			ViewTransform.ComputeOrbitMatrix().InverseFast().Rotator() : ViewTransform.GetRotation();
+		CameraStateOut.Orientation = ViewRotation.Quaternion();
+	}
+
 	/** 
 	 * Change for undoing/redoing displayed layer changes.
 	 */
@@ -342,11 +361,63 @@ void UUVEditorMode::Exit()
 	Super::Exit();
 }
 
+void UUVEditorMode::InitializeContexts(FEditorViewportClient& LivePreviewViewportClient, FAssetEditorModeManager& LivePreviewModeManager, 
+	UUVToolViewportButtonsAPI& ViewportButtonsAPI)
+{
+	using namespace UVEditorModeLocals;
+
+	UContextObjectStore* ContextStore = GetInteractiveToolsContext()->ToolManager->GetContextObjectStore();
+
+	LivePreviewWorld = LivePreviewModeManager.GetPreviewScene()->GetWorld();
+	UUVToolLivePreviewAPI* LivePreviewAPI = NewObject<UUVToolLivePreviewAPI>();
+	LivePreviewAPI->Initialize(
+		LivePreviewWorld,
+		LivePreviewModeManager.GetInteractiveToolsContext()->InputRouter,
+		[this, LivePreviewViewportClientPtr = &LivePreviewViewportClient](FViewCameraState& CameraStateOut) {
+			GetCameraState(*LivePreviewViewportClientPtr, CameraStateOut); 
+		});
+	ContextStore->AddContextObject(LivePreviewAPI);
+
+	UUVToolEmitChangeAPI* EmitChangeAPI = NewObject<UUVToolEmitChangeAPI>();
+	EmitChangeAPI->Initialize(GetInteractiveToolsContext()->ToolManager);
+	ContextStore->AddContextObject(EmitChangeAPI);
+
+	UUVToolAssetAndChannelAPI* AssetAndLayerAPI = NewObject<UUVToolAssetAndChannelAPI>();
+	AssetAndLayerAPI->RequestChannelVisibilityChangeFunc = [this](const TArray<int32>& LayerPerAsset, bool bEmitUndoTransaction) {
+		SetDisplayedUVChannels(LayerPerAsset, bEmitUndoTransaction);
+	};
+	AssetAndLayerAPI->NotifyOfAssetChannelCountChangeFunc = [this](int32 AssetID) {
+		// Don't currently need to do anything because the layer selection menu gets populated
+		// from scratch each time that it's opened.
+	};
+	AssetAndLayerAPI->GetCurrentChannelVisibilityFunc = [this]() {
+		TArray<int32> VisibleLayers;
+		VisibleLayers.SetNum(ToolTargets.Num());
+		for (int32 AssetID = 0; AssetID < ToolTargets.Num(); ++AssetID)
+		{
+			VisibleLayers[AssetID] = GetDisplayedChannel(AssetID);
+		}
+		return VisibleLayers;
+	};
+	ContextStore->AddContextObject(AssetAndLayerAPI);
+
+	UUVVisualStyleAPI* VisualStyleAPI = NewObject<UUVVisualStyleAPI>();
+	VisualStyleAPI->GetSelectionColorForAssetFunc = [this](int32 AssetID) {
+		return GetSelectionColorByTargetIndex(AssetID);
+	};
+	ContextStore->AddContextObject(VisualStyleAPI);
+
+	ContextStore->AddContextObject(&ViewportButtonsAPI);
+}
+
 void UUVEditorMode::InitializeTargets(const TArray<TObjectPtr<UObject>>& AssetsIn,
 	const TArray<FTransform>& TransformsIn)
 {
 	using namespace UVEditorModeLocals;
 
+	// InitializeContexts needs to have been called first so that we have the 3d preview world ready.
+	check(LivePreviewWorld);
+
 	OriginalObjectsToEdit = AssetsIn;
 	Transforms = TransformsIn;
 
@@ -354,13 +425,6 @@ void UUVEditorMode::InitializeTargets(const TArray<TObjectPtr<UObject>>& AssetsI
 	UUVEditorSubsystem* UVSubsystem = GEditor->GetEditorSubsystem<UUVEditorSubsystem>();
 	UVSubsystem->BuildTargets(AssetsIn, GetToolTargetRequirements(), ToolTargets);
 
-	// For creating the actual input objects, we'll need pointers both to the 2d unwrap world and the
-	// 3d preview world. We already have the 2d world in GetWorld(). Get the 3d one.
-	UContextObjectStore* ContextStore = GetInteractiveToolsContext()->ToolManager->GetContextObjectStore();
-	UUVToolLivePreviewAPI* LivePreviewAPI = ContextStore->FindContext<UUVToolLivePreviewAPI>();
-	check(LivePreviewAPI);
-	LivePreviewWorld = LivePreviewAPI->GetLivePreviewWorld();
-
 	// Collect the 3d dynamic meshes from targets. There will always be one for each asset, and the AssetID
 	// of each asset will be the index into these arrays. Individual input objects (representing a specific
 	// UV layer), will point to these existing 3d meshes.
@@ -471,39 +535,6 @@ void UUVEditorMode::InitializeTargets(const TArray<TObjectPtr<UObject>>& AssetsI
 	// Prep things for layer/channel selection
 	InitializeAssetNames(ToolTargets, AssetNames);
 	PendingUVLayerIndex.SetNumZeroed(ToolTargets.Num());
-
-	// Prep channel/layer change API
-	UUVToolAssetAndChannelAPI* AssetAndLayerAPI = ContextStore->FindContext<UUVToolAssetAndChannelAPI>();
-	if (AssetAndLayerAPI)
-	{
-		AssetAndLayerAPI->RequestChannelVisibilityChangeFunc = [this](const TArray<int32>& LayerPerAsset, bool bEmitUndoTransaction) {
-			this->SetDisplayedUVChannels(LayerPerAsset, bEmitUndoTransaction);
-		};
-
-		AssetAndLayerAPI->NotifyOfAssetChannelCountChangeFunc = [this](int32 AssetID) {
-			// Don't currently need to do anything because the layer selection menu gets populated
-			// from scratch each time that it's opened.
-		};
-
-		AssetAndLayerAPI->GetCurrentChannelVisibilityFunc = [this]() {
-			TArray<int32> VisibleLayers;
-			VisibleLayers.SetNum(ToolTargets.Num());
-			for (int32 AssetID = 0; AssetID < ToolTargets.Num(); ++AssetID)
-			{
-				VisibleLayers[AssetID] = GetDisplayedChannel(AssetID);
-			}
-			return VisibleLayers;
-		};
-	}
-
-	// Prep VisualStyle API
-	UUVVisualStyleAPI* VisualStyleAPI = ContextStore->FindContext<UUVVisualStyleAPI>();
-	if (VisualStyleAPI)
-	{
-		VisualStyleAPI->GetSelectionColorForAssetFunc = [this](int32 AssetID) {
-			return GetSelectionColorByTargetIndex(AssetID);
-		};
-	}
 }
 
 FLinearColor UUVEditorMode::GetTriangleColorByTargetIndex(int32 TargetIndex) const
diff --git a/Engine/Plugins/Experimental/UVEditor/Source/UVEditor/Private/UVEditorToolkit.cpp b/Engine/Plugins/Experimental/UVEditor/Source/UVEditor/Private/UVEditorToolkit.cpp
--- a/Engine/Plugins/Experimental/UVEditor/Source/UVEditor/Private/UVEditorToolkit.cpp
+++ b/Engine/Plugins/Experimental/UVEditor/Source/UVEditor/Private/UVEditorToolkit.cpp
@@ -37,24 +37,6 @@
 const FName FUVEditorToolkit::InteractiveToolsPanelTabID(TEXT("UVEditor_InteractiveToolsTab"));
 const FName FUVEditorToolkit::LivePreviewTabID(TEXT("UVEditor_LivePreviewTab"));
 
-namespace FUVEditorToolkitLocals {
-	void GetCameraState(const FEditorViewportClient& ViewportClientIn, FViewCameraState& CameraStateOut)
-	{
-		FViewportCameraTransform ViewTransform = ViewportClientIn.GetViewTransform();
-		CameraStateOut.bIsOrthographic = false;
-		CameraStateOut.bIsVR = false;
-		CameraStateOut.Position = ViewTransform.GetLocation();
-		CameraStateOut.HorizontalFOVDegrees = ViewportClientIn.ViewFOV;
-		CameraStateOut.AspectRatio = ViewportClientIn.AspectRatio;
-
-		// if using Orbit camera, the rotation in the ViewTransform is not the current camera rotation, it
-		// is set to a different rotation based on the Orbit. So we have to convert back to camera rotation.
-		FRotator ViewRotation = (ViewportClientIn.bUsingOrbitCamera) ?
-			ViewTransform.ComputeOrbitMatrix().InverseFast().Rotator() : ViewTransform.GetRotation();
-		CameraStateOut.Orientation = ViewRotation.Quaternion();
-	}
-}
-
 FUVEditorToolkit::FUVEditorToolkit(UAssetEditor* InOwningAssetEditor)
 	: FBaseAssetToolkit(InOwningAssetEditor)
 {
@@ -127,6 +109,11 @@ FUVEditorToolkit::FUVEditorToolkit(UAssetEditor* InOwningAssetEditor)
 		return SNew(SUVEditor3DViewport, InArgs)
 			.EditorViewportClient(LivePreviewViewportClient);
 	};
+
+	// This API object serves as a communication point between the viewport toolbars and the tools.
+	// We create it here so that we can pass it both into the 2d viewport and when we initialize 
+	// the mode.
+	ViewportButtonsAPI = NewObject<UUVToolViewportButtonsAPI>();
 }
 
 FUVEditorToolkit::~FUVEditorToolkit()
@@ -349,7 +336,8 @@ TSharedPtr<FEditorViewportClient> FUVEditorToolkit::CreateEditorViewportClient()
 	// namely ViewportType.
 	// Instead, we do viewport client adjustment in PostInitAssetEditor().
 	check(EditorModeManager.IsValid());
-	return MakeShared<FUVEditor2DViewportClient>(EditorModeManager.Get(), UnwrapScene.Get());
+	return MakeShared<FUVEditor2DViewportClient>(EditorModeManager.Get(), UnwrapScene.Get(), 
+		UVEditor2DViewport, ViewportButtonsAPI);
 }
 
 // Called from FBaseAssetToolkit::CreateWidgets. The delegate call path goes through FAssetEditorToolkit::InitAssetEditor
@@ -358,7 +346,7 @@ AssetEditorViewportFactoryFunction FUVEditorToolkit::GetViewportDelegate()
 {
 	AssetEditorViewportFactoryFunction TempViewportDelegate = [this](FAssetEditorViewportConstructionArgs InArgs)
 	{
-		return SNew(SUVEditor2DViewport, InArgs)
+		return SAssignNew(UVEditor2DViewport, SUVEditor2DViewport, InArgs)
 			.EditorViewportClient(ViewportClient);
 	};
 
@@ -396,27 +384,9 @@ void FUVEditorToolkit::PostInitAssetEditor()
 		EditorModeManager->GetActiveScriptableMode(UUVEditorMode::EM_UVEditorModeId));
 	check(UVMode);
 
-	// The mode will need to be able to get to the live preview world and input.
-	UContextObjectStore* ContextStore = EditorModeManager->GetInteractiveToolsContext()->ToolManager->GetContextObjectStore();
-	UUVToolLivePreviewAPI* LivePreviewAPI = NewObject<UUVToolLivePreviewAPI>();
-	LivePreviewAPI->Initialize(LivePreviewScene->GetWorld(), LivePreviewInputRouter,
-		[this](FViewCameraState& CameraStateOut) {
-			FUVEditorToolkitLocals::GetCameraState(*LivePreviewViewportClient, CameraStateOut); });
-	ContextStore->AddContextObject(LivePreviewAPI);
-
-	// It is debatable where the undo/redo api should get initialized, but this seems like a 
-	// reasonable place since we're setting up other things for the mode for the first time.
-	UUVToolEmitChangeAPI* EmitChangeAPI = NewObject<UUVToolEmitChangeAPI>();
-	EmitChangeAPI->Initialize(EditorModeManager->GetInteractiveToolsContext()->ToolManager);
-	ContextStore->AddContextObject(EmitChangeAPI);
-
-	UUVToolAssetAndChannelAPI* AssetAndLayerAPI = NewObject<UUVToolAssetAndChannelAPI>();
-	ContextStore->AddContextObject(AssetAndLayerAPI);
-
-	UUVVisualStyleAPI* VisualStyleAPI = NewObject<UUVVisualStyleAPI>();
-	ContextStore->AddContextObject(VisualStyleAPI);
+	// The mode will need to be able to get to the live preview world, camera, input router, and viewport buttons.
+	UVMode->InitializeContexts(*LivePreviewViewportClient, *LivePreviewEditorModeManager, *ViewportButtonsAPI);
 
-	// Initialize mode state.
 	TArray<TObjectPtr<UObject>> ObjectsToEdit;
 	OwningAssetEditor->GetObjectsToEdit(ObjectsToEdit);
 
diff --git a/Engine/Plugins/Experimental/UVEditor/Source/UVEditor/Public/UVEditor2DViewportClient.h b/Engine/Plugins/Experimental/UVEditor/Source/UVEditor/Public/UVEditor2DViewportClient.h
--- a/Engine/Plugins/Experimental/UVEditor/Source/UVEditor/Public/UVEditor2DViewportClient.h
+++ b/Engine/Plugins/Experimental/UVEditor/Source/UVEditor/Public/UVEditor2DViewportClient.h
@@ -9,6 +9,8 @@
 #include "UVEditor2DViewportBehaviorTargets.h" // FUVEditor2DScrollBehaviorTarget, FUVEditor2DMouseWheelZoomBehaviorTarget
 #include "UVToolContextObjects.h" // UUVToolViewportButtonsAPI::ESelectionMode
 
+class UUVToolViewportButtonsAPI;
+
 /**
  * Client used to display a 2D view of the UV's, implemented by using a perspective viewport with a locked
  * camera.
@@ -16,8 +18,8 @@
 class UVEDITOR_API FUVEditor2DViewportClient : public FEditorViewportClient, public IInputBehaviorSource
 {
 public:
-	FUVEditor2DViewportClient(FEditorModeTools* InModeTools, FPreviewScene* InPreviewScene = nullptr,
-		const TWeakPtr<SEditorViewport>& InEditorViewportWidget = nullptr);
+	FUVEditor2DViewportClient(FEditorModeTools* InModeTools, FPreviewScene* InPreviewScene,
+		const TWeakPtr<SEditorViewport>& InEditorViewportWidget, UUVToolViewportButtonsAPI* ViewportButtonsAPI);
 
 	virtual ~FUVEditor2DViewportClient() {}
 
diff --git a/Engine/Plugins/Experimental/UVEditor/Source/UVEditor/Public/UVEditorMode.h b/Engine/Plugins/Experimental/UVEditor/Source/UVEditor/Public/UVEditorMode.h
--- a/Engine/Plugins/Experimental/UVEditor/Source/UVEditor/Public/UVEditorMode.h
+++ b/Engine/Plugins/Experimental/UVEditor/Source/UVEditor/Public/UVEditorMode.h
@@ -12,6 +12,7 @@
 
 PREDECLARE_GEOMETRY(class FDynamicMesh3);
 
+class FEditorViewportClient;
 class FToolCommandChange;
 class IUVUnwrapDynamicMesh;
 class UMeshElementsVisualizer;
@@ -24,6 +25,7 @@ class UMeshOpPreviewWithBackgroundCompute;
 class UUVToolStateObjectStore;
 class UUVEditorToolMeshInput;
 class UUVEditorBackgroundPreview;
+class UUVToolViewportButtonsAPI;
 
 
 /**
@@ -53,6 +55,10 @@ public:
 	 */
 	static double GetUVMeshScalingFactor() { return 1000; }
 
+	// Both initialization functions must be called for things to function properly. InitializeContexts should
+	// be done first so that the 3d preview world is ready for creating meshes in InitializeTargets.
+	void InitializeContexts(FEditorViewportClient& LivePreviewViewportClient, FAssetEditorModeManager& LivePreviewModeManager,
+		UUVToolViewportButtonsAPI& ViewportButtonsAPI);
 	void InitializeTargets(const TArray<TObjectPtr<UObject>>& AssetsIn, const TArray<FTransform>& TransformsIn);
 
 	// Public for use by undo/redo. Otherwise should use RequestUVChannelChange
@@ -114,9 +120,6 @@ public:
 	UPROPERTY()
 	TObjectPtr<UUVEditorBackgroundPreview> BackgroundVisualization;
 
-	// Our UV tools use contexts that are set up at asset editor level, so scope needs to be higher level than mode
-	EToolsContextScope GetDefaultToolScope() const override { return EToolsContextScope::Editor; }
-
 protected:
 
 	// UEdMode overrides
diff --git a/Engine/Plugins/Experimental/UVEditor/Source/UVEditor/Public/UVEditorToolkit.h b/Engine/Plugins/Experimental/UVEditor/Source/UVEditor/Public/UVEditorToolkit.h
--- a/Engine/Plugins/Experimental/UVEditor/Source/UVEditor/Public/UVEditorToolkit.h
+++ b/Engine/Plugins/Experimental/UVEditor/Source/UVEditor/Public/UVEditorToolkit.h
@@ -7,6 +7,7 @@
 #include "Tools/BaseAssetToolkit.h"
 
 class FAdvancedPreviewScene;
+class FUVEditorModeUILayer;
 class FUVEditorPreviewModeManager;
 class IDetailsView;
 class SDockTab;
@@ -14,7 +15,7 @@ class SBorder;
 class SWidget;
 class UInteractiveToolsContext;
 class UInputRouter;
-class FUVEditorModeUILayer;
+class UUVToolViewportButtonsAPI;
 
 /**
  * The toolkit is supposed to act as the UI manager for the asset editor. It's responsible 
@@ -84,5 +85,8 @@ protected:
 	TSharedPtr<FAssetEditorModeManager> LivePreviewEditorModeManager;
 	TObjectPtr<UInputRouter> LivePreviewInputRouter = nullptr;
 
+	TWeakPtr<SEditorViewport> UVEditor2DViewport;
+	UUVToolViewportButtonsAPI* ViewportButtonsAPI = nullptr;
+
 	TSharedPtr<FUVEditorModeUILayer> ModeUILayer;
 };
\ No newline at end of file
diff --git a/Engine/Source/Editor/UnrealEd/Public/Tools/UEdMode.h b/Engine/Source/Editor/UnrealEd/Public/Tools/UEdMode.h
--- a/Engine/Source/Editor/UnrealEd/Public/Tools/UEdMode.h
+++ b/Engine/Source/Editor/UnrealEd/Public/Tools/UEdMode.h
@@ -206,7 +206,7 @@ public:
 	 * Default Scope for InteractiveToolsContext API functions, eg RegisterTool(), GetToolManager(), GetInteractiveToolsContext().
 	 * See EToolsContextScope for details. Defaults to Editor scope.
 	 */
-	virtual EToolsContextScope GetDefaultToolScope() const { return EToolsContextScope::EdMode; }
+	EToolsContextScope GetDefaultToolScope() const { return EToolsContextScope::EdMode; }
 
 	/**
 	 * @return active ToolManager for the desired (or default) ToolsContext Scope
