UVEditor: Add hotkeys for selection modes (1-5). Fix bug where gizmo rotation would always reset after rotating.

#rnx

diff --git a/Engine/Plugins/Experimental/UVEditor/Source/UVEditor/Private/UVEditorCommands.cpp b/Engine/Plugins/Experimental/UVEditor/Source/UVEditor/Private/UVEditorCommands.cpp
--- a/Engine/Plugins/Experimental/UVEditor/Source/UVEditor/Private/UVEditorCommands.cpp
+++ b/Engine/Plugins/Experimental/UVEditor/Source/UVEditor/Private/UVEditorCommands.cpp
@@ -42,11 +42,11 @@ void FUVEditorCommands::RegisterCommands()
 	UI_COMMAND(CancelOrCompleteActiveTool, "Cancel", "Cancel the active tool", EUserInterfaceActionType::Button, FInputChord(EKeys::Escape));
 
 	// These get used in viewport buttons
-	UI_COMMAND(VertexSelection, "Vertex Selection", "Select vertices", EUserInterfaceActionType::ToggleButton, FInputChord());
-	UI_COMMAND(EdgeSelection, "Edge Selection", "Select edges", EUserInterfaceActionType::ToggleButton, FInputChord());
-	UI_COMMAND(TriangleSelection, "Triangle Selection", "Select triangles", EUserInterfaceActionType::ToggleButton, FInputChord());
-	UI_COMMAND(IslandSelection, "Island Selection", "Select connected islands", EUserInterfaceActionType::ToggleButton, FInputChord());
-	UI_COMMAND(FullMeshSelection, "Full Mesh Selection", "Select whole mesh", EUserInterfaceActionType::ToggleButton, FInputChord());
+	UI_COMMAND(VertexSelection, "Vertex Selection", "Select vertices", EUserInterfaceActionType::ToggleButton, FInputChord(EKeys::One));
+	UI_COMMAND(EdgeSelection, "Edge Selection", "Select edges", EUserInterfaceActionType::ToggleButton, FInputChord(EKeys::Two));
+	UI_COMMAND(TriangleSelection, "Triangle Selection", "Select triangles", EUserInterfaceActionType::ToggleButton, FInputChord(EKeys::Three));
+	UI_COMMAND(IslandSelection, "Island Selection", "Select connected islands", EUserInterfaceActionType::ToggleButton, FInputChord(EKeys::Four));
+	UI_COMMAND(FullMeshSelection, "Full Mesh Selection", "Select whole mesh", EUserInterfaceActionType::ToggleButton, FInputChord(EKeys::Five));
 
 	UI_COMMAND(EnableOrbitCamera, "Orbit", "Enable Orbit Camera", EUserInterfaceActionType::ToggleButton, FInputChord());
 	UI_COMMAND(EnableFlyCamera, "Fly", "Enable Fly Camera", EUserInterfaceActionType::ToggleButton, FInputChord());
diff --git a/Engine/Plugins/Experimental/UVEditor/Source/UVEditorTools/Private/UVSelectTool.cpp b/Engine/Plugins/Experimental/UVEditor/Source/UVEditorTools/Private/UVSelectTool.cpp
--- a/Engine/Plugins/Experimental/UVEditor/Source/UVEditorTools/Private/UVSelectTool.cpp
+++ b/Engine/Plugins/Experimental/UVEditor/Source/UVEditorTools/Private/UVSelectTool.cpp
@@ -398,6 +398,10 @@ void UUVSelectTool::Setup()
 	{
 		Targets[i]->OnCanonicalModified.AddWeakLambda(this, [this, i]
 		(UUVEditorToolMeshInput* InputObject, const UUVEditorToolMeshInput::FCanonicalModifiedInfo& Info) {
+			if (bIgnoreOnCanonicalChange) // Used to avoid reacting to broadcasts that we ourselves caused
+			{
+				return;
+			}
 			UpdateSelectionEidsAfterMeshChange(*SelectionMechanic, &CurrentSelectionVidPairs);
 			UpdateGizmo();
 			SelectionMechanic->RebuildDrawnElements(TransformGizmo->GetGizmoTransform());
@@ -845,9 +849,16 @@ void UUVSelectTool::GizmoTransformEnded(UTransformProxy* Proxy)
 	// One final attempt to apply transforms if OnTick hasn't happened yet
 	ApplyGizmoTransform();
 
-	// Both previews must already be updated, so only need to update canonical
-	Targets[SelectionTargetIndex]->UpdateCanonicalFromPreviews(&MovingVids, 
-		UUVEditorToolMeshInput::NONE_CHANGED_ARG);
+	// Both previews must already be updated, so only need to update canonical. 
+	{
+		// We don't want to react to the ensuing broadcast so that we don't lose the gizmo rotation. We could just 
+		// not broadcast (and update related structures, i.e. trees, ourselves), but conceptually it's better to 
+		// broadcast the change since we did change the canonicals.
+		TGuardValue<bool> IgnoreOnCanonicalBroadcast(bIgnoreOnCanonicalChange, true); // sets to true, restores at end
+
+		Targets[SelectionTargetIndex]->UpdateCanonicalFromPreviews(&MovingVids,
+			UUVEditorToolMeshInput::NONE_CHANGED_ARG);
+	}
 
 	const FText TransactionName(LOCTEXT("DragCompleteTransactionName", "Move Items"));
 	EmitChangeAPI->EmitToolIndependentChange(ChangeRouter, MakeUnique<UVSelectToolLocals::FGizmoMeshChange>(
diff --git a/Engine/Plugins/Experimental/UVEditor/Source/UVEditorTools/Public/UVSelectTool.h b/Engine/Plugins/Experimental/UVEditor/Source/UVEditorTools/Public/UVSelectTool.h
--- a/Engine/Plugins/Experimental/UVEditor/Source/UVEditorTools/Public/UVSelectTool.h
+++ b/Engine/Plugins/Experimental/UVEditor/Source/UVEditorTools/Public/UVSelectTool.h
@@ -179,6 +179,7 @@ protected:
 	UE::Geometry::FFrame3d InitialGizmoFrame;
 	FTransform UnappliedGizmoTransform;
 	bool bInDrag = false;
+	bool bIgnoreOnCanonicalChange = false;
 	bool bGizmoTransformNeedsApplication = false;
 
 	TArray<int32> MovingVids;
