Non-Unity Compiler Fixes

diff --git a/Engine/Plugins/Experimental/UVEditor/Source/UVEditor/Public/UVEditorMode.h b/Engine/Plugins/Experimental/UVEditor/Source/UVEditor/Public/UVEditorMode.h
--- a/Engine/Plugins/Experimental/UVEditor/Source/UVEditor/Public/UVEditorMode.h
+++ b/Engine/Plugins/Experimental/UVEditor/Source/UVEditor/Public/UVEditorMode.h
@@ -13,6 +13,7 @@
 PREDECLARE_GEOMETRY(class FDynamicMesh3);
 
 class FEditorViewportClient;
+class FAssetEditorModeManager;
 class FToolCommandChange;
 class IUVUnwrapDynamicMesh;
 class UMeshElementsVisualizer;
