UVEditor: Remove the material override properties from the Auto UV tool, since they conflict with the existing editor settings for background/live preview material display.

#rnx

diff --git a/Engine/Plugins/Experimental/UVEditor/Source/UVEditorToolsEditorOnly/Private/UVEditorParameterizeMeshTool.cpp b/Engine/Plugins/Experimental/UVEditor/Source/UVEditorToolsEditorOnly/Private/UVEditorParameterizeMeshTool.cpp
--- a/Engine/Plugins/Experimental/UVEditor/Source/UVEditorToolsEditorOnly/Private/UVEditorParameterizeMeshTool.cpp
+++ b/Engine/Plugins/Experimental/UVEditor/Source/UVEditorToolsEditorOnly/Private/UVEditorParameterizeMeshTool.cpp
@@ -69,18 +69,6 @@ void UUVEditorParameterizeMeshTool::Setup()
 	AddToolPropertySource(PatchBuilderProperties);
 	SetToolPropertySourceEnabled(PatchBuilderProperties, true);
 
-	MaterialSettings = NewObject<UExistingMeshMaterialProperties>(this);
-	MaterialSettings->MaterialMode = ESetMeshMaterialMode::Checkerboard;
-	MaterialSettings->RestoreProperties(this, TEXT("ModelingUVTools"));
-	MaterialSettings->UpdateUVChannels(0, {TEXT("UV 0")}); // We only ever have one UV layer for the applied preview.
-	AddToolPropertySource(MaterialSettings);
-	// force update
-	MaterialSettings->UpdateMaterials();
-	for (TObjectPtr<UUVEditorToolMeshInput> Target : Targets)
-	{
-		Target->AppliedPreview->OverrideMaterial = MaterialSettings->GetActiveOverrideMaterial();
-	}
-
 	Factories.SetNum(Targets.Num());
 	for (int32 TargetIndex = 0; TargetIndex < Targets.Num(); ++TargetIndex)
 	{
@@ -112,19 +100,11 @@ void UUVEditorParameterizeMeshTool::OnPropertyModified(UObject* PropertySet, FPr
 {
 	TRACE_CPUPROFILER_EVENT_SCOPE(UVEditorParameterizeMeshTool_OnPropertyModified);
 
-	if (PropertySet != MaterialSettings)
-	{
-		for (TObjectPtr<UUVEditorToolMeshInput> Target : Targets)
-		{
-			Target->AppliedPreview->InvalidateResult();
-		}
-	}
-
-	MaterialSettings->UpdateMaterials();
 	for (TObjectPtr<UUVEditorToolMeshInput> Target : Targets)
 	{
-		Target->AppliedPreview->OverrideMaterial = MaterialSettings->GetActiveOverrideMaterial();
+		Target->AppliedPreview->InvalidateResult();
 	}
+	
 }
 
 
@@ -148,7 +128,6 @@ void UUVEditorParameterizeMeshTool::Shutdown(EToolShutdownType ShutdownType)
 	TRACE_CPUPROFILER_EVENT_SCOPE(UVEditorParameterizeMeshTool_Shutdown);
 
 	Settings->SaveProperties(this);
-	MaterialSettings->SaveProperties(this, TEXT("ModelingUVTools"));
 	UVAtlasProperties->SaveProperties(this);
 	XAtlasProperties->SaveProperties(this);
 	PatchBuilderProperties->SaveProperties(this);
diff --git a/Engine/Plugins/Experimental/UVEditor/Source/UVEditorToolsEditorOnly/Public/UVEditorParameterizeMeshTool.h b/Engine/Plugins/Experimental/UVEditor/Source/UVEditorToolsEditorOnly/Public/UVEditorParameterizeMeshTool.h
--- a/Engine/Plugins/Experimental/UVEditor/Source/UVEditorToolsEditorOnly/Public/UVEditorParameterizeMeshTool.h
+++ b/Engine/Plugins/Experimental/UVEditor/Source/UVEditorToolsEditorOnly/Public/UVEditorParameterizeMeshTool.h
@@ -85,9 +85,6 @@ protected:
 	UPROPERTY()
 	TObjectPtr<UParameterizeMeshToolPatchBuilderProperties> PatchBuilderProperties = nullptr;
 
-	UPROPERTY()
-	TObjectPtr<UExistingMeshMaterialProperties> MaterialSettings = nullptr;
-
 	UPROPERTY()
 	TArray<TObjectPtr<UParameterizeMeshOperatorFactory>> Factories;
 
