UVEditor: Move visual and style settings (color, depth offsets, etc) to a new style container class, keeping everything in one place for ease of modification in the future. Also adjusts the grid axis lines to not scale up when zooming out in the 2D viewport

#rnx

diff --git a/Engine/Plugins/Experimental/UVEditor/Source/UVEditor/Private/UVEditor2DViewportClient.cpp b/Engine/Plugins/Experimental/UVEditor/Source/UVEditor/Private/UVEditor2DViewportClient.cpp
--- a/Engine/Plugins/Experimental/UVEditor/Source/UVEditor/Private/UVEditor2DViewportClient.cpp
+++ b/Engine/Plugins/Experimental/UVEditor/Source/UVEditor/Private/UVEditor2DViewportClient.cpp
@@ -10,6 +10,7 @@
 #include "Drawing/MeshDebugDrawing.h"
 #include "FrameTypes.h"
 #include "UVEditorMode.h"
+#include "UVEditorUXSettings.h"
 
 FUVEditor2DViewportClient::FUVEditor2DViewportClient(FEditorModeTools* InModeTools,
 	FPreviewScene* InPreviewScene, const TWeakPtr<SEditorViewport>& InEditorViewportWidget,
@@ -41,8 +42,8 @@ FUVEditor2DViewportClient::FUVEditor2DViewportClient(FEditorModeTools* InModeToo
 	BehaviorSet->Add(ScrollBehavior);
 
 	ZoomBehaviorTarget = MakeUnique<FUVEditor2DMouseWheelZoomBehaviorTarget>(this);
-	ZoomBehaviorTarget->SetCameraFarPlaneWorldZ(CameraFarPlaneWorldZ);
-	ZoomBehaviorTarget->SetCameraNearPlaneProportionZ(CameraNearPlaneProportionZ);
+	ZoomBehaviorTarget->SetCameraFarPlaneWorldZ(FUVEditorUXSettings::CameraFarPlaneWorldZ);
+	ZoomBehaviorTarget->SetCameraNearPlaneProportionZ(FUVEditorUXSettings::CameraNearPlaneProportionZ);
 	UMouseWheelInputBehavior* ZoomBehavior = NewObject<UMouseWheelInputBehavior>();
 	ZoomBehavior->Initialize(ZoomBehaviorTarget.Get());
 	ZoomBehavior->SetDefaultPriority(DEFAULT_VIEWPORT_BEHAVIOR_PRIORITY);	
@@ -85,16 +86,6 @@ void FUVEditor2DViewportClient::Draw(const FSceneView* View, FPrimitiveDrawInter
 	// Basic scaling amount
 	const float UVScale = UUVEditorMode::GetUVMeshScalingFactor();
 	
-	// Line thickness parameters
-	const float AxisThickness = 2;
-	const float GridMajorThickness = 1.0;
-
-    // Line color scheme parameters
-	const FLinearColor XAxisColor = FLinearColor::Red;
-	const FLinearColor YAxisColor = FLinearColor::Green;
-	const FLinearColor GridMajorColor = FLinearColor::Gray;
-	const FLinearColor GridMinorColor = FLinearColor::Gray;
-
 	// Determine important geometry of the viewport for creating grid lines
 	FVector WorldCenterPoint( 0,0,0 );
 	FVector4 WorldToScreenCenter = View->WorldToScreen(WorldCenterPoint);	
@@ -114,23 +105,23 @@ void FUVEditor2DViewportClient::Draw(const FSceneView* View, FPrimitiveDrawInter
 		UE::Geometry::FFrame3f LocalFrame((FVector3f)ViewLoc);
 		FTransform Transform;
 		TArray<FColor> Colors;
-		Colors.Push(GridMajorColor.ToRGBE());
-		Colors.Push(GridMinorColor.ToRGBE());
+		Colors.Push(FUVEditorUXSettings::GridMajorColor);
+		Colors.Push(FUVEditorUXSettings::GridMinorColor);
 		MeshDebugDraw::DrawHierarchicalGrid(UVScale, ZoomFactor / UVScale,
 			500, // Maximum density of lines to draw per level before skipping the level
 			WorldBoundsMax, WorldBoundsMin,
-			3, // Number of levels to draw
-			4, // Number of subdivisions per level
+			FUVEditorUXSettings::GridLevels, // Number of levels to draw
+			FUVEditorUXSettings::GridSubdivisionsPerLevel, // Number of subdivisions per level
 			Colors,
-			LocalFrame, GridMajorThickness, true,
+			LocalFrame, FUVEditorUXSettings::GridMajorThickness, true,
 			PDI, Transform);
 	}
 
-	float AxisExtent = FMathf::Max(UVScale, FMathf::Min(WorldBoundsMax.Y, WorldBoundsMax.X));
+	float AxisExtent = UVScale;
 
 	// Draw colored axis lines
-	PDI->DrawLine(FVector(0, 0, 0), FVector(AxisExtent, 0, 0), FLinearColor::Red, SDPG_World, AxisThickness, 0, true);
-	PDI->DrawLine(FVector(0, 0, 0), FVector(0, AxisExtent, 0), FLinearColor::Green, SDPG_World, AxisThickness, 0, true);
+	PDI->DrawLine(FVector(0, 0, 0), FVector(AxisExtent, 0, 0), FUVEditorUXSettings::XAxisColor, SDPG_World, FUVEditorUXSettings::AxisThickness, 0, true);
+	PDI->DrawLine(FVector(0, 0, 0), FVector(0, AxisExtent, 0), FUVEditorUXSettings::YAxisColor, SDPG_World, FUVEditorUXSettings::AxisThickness, 0, true);
 
 	// TODO: Draw a little UV axis thing in the lower left, like the XYZ things that normal viewports have.
 
diff --git a/Engine/Plugins/Experimental/UVEditor/Source/UVEditor/Private/UVEditorBackgroundPreview.cpp b/Engine/Plugins/Experimental/UVEditor/Source/UVEditor/Private/UVEditorBackgroundPreview.cpp
--- a/Engine/Plugins/Experimental/UVEditor/Source/UVEditor/Private/UVEditorBackgroundPreview.cpp
+++ b/Engine/Plugins/Experimental/UVEditor/Source/UVEditor/Private/UVEditorBackgroundPreview.cpp
@@ -7,6 +7,7 @@
 #include "Drawing/MeshWireframeComponent.h"
 #include "ToolSetupUtil.h"
 #include "Async/Async.h"
+#include "UVEditorUXSettings.h"
 
 using namespace UE::Geometry;
 
@@ -99,7 +100,7 @@ void UUVEditorBackgroundPreview::UpdateBackground()
 			ensure(false);
 	}
 
-	BackgroundMaterial->SetScalarParameterValue(TEXT("BackgroundPixelDepthOffset"), 0);
+	BackgroundMaterial->SetScalarParameterValue(TEXT("BackgroundPixelDepthOffset"), FUVEditorUXSettings::BackgroundQuadDepthOffset);
 	OnBackgroundMaterialChange.Broadcast(BackgroundMaterial);
 	BackgroundComponent->Clear();
 
diff --git a/Engine/Plugins/Experimental/UVEditor/Source/UVEditor/Private/UVEditorMode.cpp b/Engine/Plugins/Experimental/UVEditor/Source/UVEditor/Private/UVEditorMode.cpp
--- a/Engine/Plugins/Experimental/UVEditor/Source/UVEditor/Private/UVEditorMode.cpp
+++ b/Engine/Plugins/Experimental/UVEditor/Source/UVEditor/Private/UVEditorMode.cpp
@@ -36,6 +36,7 @@
 #include "UVToolContextObjects.h"
 #include "UVEditorBackgroundPreview.h"
 #include "Editor.h"
+#include "UVEditorUXSettings.h"
 
 #define LOCTEXT_NAMESPACE "UUVEditorMode"
 
@@ -397,12 +398,6 @@ void UUVEditorMode::InitializeContexts(FEditorViewportClient& LivePreviewViewpor
 	};
 	ContextStore->AddContextObject(AssetAndLayerAPI);
 
-	UUVVisualStyleAPI* VisualStyleAPI = NewObject<UUVVisualStyleAPI>();
-	VisualStyleAPI->GetSelectionColorForAssetFunc = [this](int32 AssetID) {
-		return GetSelectionColorByTargetIndex(AssetID);
-	};
-	ContextStore->AddContextObject(VisualStyleAPI);
-
 	ContextStore->AddContextObject(&ViewportButtonsAPI);
 }
 
@@ -486,20 +481,20 @@ void UUVEditorMode::InitializeTargets(const TArray<TObjectPtr<UObject>>& AssetsI
 		ToolInputObject->UnwrapPreview->PreviewMesh->SetMaterial(
 			0, ToolSetupUtil::GetCustomTwoSidedDepthOffsetMaterial(
 				GetToolManager(),
-				(FLinearColor)GetTriangleColorByTargetIndex(AssetID),
-				TriangleDepthOffset,
-				TriangleOpacity));
+				FUVEditorUXSettings::GetTriangleColorByTargetIndex(AssetID),
+				FUVEditorUXSettings::UnwrapTriangleDepthOffset,
+				FUVEditorUXSettings::UnwrapTriangleOpacity));
 
 		// Set up the wireframe display of the unwrapped mesh.
 		UMeshElementsVisualizer* WireframeDisplay = NewObject<UMeshElementsVisualizer>(this);
 		WireframeDisplay->CreateInWorld(GetWorld(), FTransform::Identity);
 
-		WireframeDisplay->Settings->DepthBias = WireframeDepthOffset;
+		WireframeDisplay->Settings->DepthBias = FUVEditorUXSettings::WireframeDepthOffset;
 		WireframeDisplay->Settings->bAdjustDepthBiasUsingMeshSize = false;
 		WireframeDisplay->Settings->bShowWireframe = true;
 		WireframeDisplay->Settings->bShowBorders = true;
-		WireframeDisplay->Settings->WireframeColor = GetWireframeColorByTargetIndex(AssetID).ToFColorSRGB();
-		WireframeDisplay->Settings->BoundaryEdgeColor = GetBoundaryColorByTargetIndex(AssetID).ToFColorSRGB(); ;
+		WireframeDisplay->Settings->WireframeColor = FUVEditorUXSettings::GetWireframeColorByTargetIndex(AssetID).ToFColorSRGB();
+		WireframeDisplay->Settings->BoundaryEdgeColor = FUVEditorUXSettings::GetBoundaryColorByTargetIndex(AssetID).ToFColorSRGB(); ;
 		WireframeDisplay->Settings->bShowUVSeams = false;
 		WireframeDisplay->Settings->bShowNormalSeams = false;
 		// These are not exposed at the visualizer level yet
@@ -533,42 +528,6 @@ void UUVEditorMode::InitializeTargets(const TArray<TObjectPtr<UObject>>& AssetsI
 	PendingUVLayerIndex.SetNumZeroed(ToolTargets.Num());
 }
 
-FLinearColor UUVEditorMode::GetTriangleColorByTargetIndex(int32 TargetIndex) const
-{
-	double GoldenAngle = 137.50776405;
-
-	FLinearColor BaseColorHSV;
-	BaseColorHSV.R = FMath::Fmod(UnwrapBaseColorHueStart + (GoldenAngle / 2.0 * TargetIndex), 360);;
-	BaseColorHSV.G = UnwrapBaseColorSaturation;
-	BaseColorHSV.B = UnwrapBaseColorValue;
-	return BaseColorHSV.HSVToLinearRGB();
-}
-
-FLinearColor UUVEditorMode::GetWireframeColorByTargetIndex(int32 TargetIndex) const
-{
-	FLinearColor BaseColorHSV = GetTriangleColorByTargetIndex(TargetIndex).LinearRGBToHSV();
-	FLinearColor WireframeColorHSV = BaseColorHSV;
-	WireframeColorHSV.R = FMath::Fmod((WireframeColorHSV.R + UnwrapWireframeHueShift), 360);
-	WireframeColorHSV.G = UnwrapWireframeSaturation;
-	WireframeColorHSV.B = UnwrapWireframeValue;
-	return WireframeColorHSV.HSVToLinearRGB();
-}
-
-FLinearColor UUVEditorMode::GetBoundaryColorByTargetIndex(int32 TargetIndex) const
-{
-	FLinearColor BaseColorHSV = GetTriangleColorByTargetIndex(TargetIndex).LinearRGBToHSV();
-	FLinearColor BoundaryColorHSV = BaseColorHSV;
-	BoundaryColorHSV.R = FMath::Fmod( (BoundaryColorHSV.R + UnwrapBoundaryHueShift), 360);
-	BoundaryColorHSV.G = UnwrapBoundarySaturation;
-	BoundaryColorHSV.B = UnwrapBoundaryValue;
-	return BoundaryColorHSV.HSVToLinearRGB();
-}
-
-FLinearColor UUVEditorMode::GetSelectionColorByTargetIndex(int32 TargetIndex) const
-{
-	return FColor::FromHex("EBA30DFF");	
-}
-
 void UUVEditorMode::EmitToolIndependentObjectChange(UObject* TargetObject, TUniquePtr<FToolCommandChange> Change, const FText& Description)
 {
 	GetInteractiveToolsContext()->GetTransactionAPI()->AppendChange(TargetObject, MoveTemp(Change), Description);
@@ -607,16 +566,15 @@ void UUVEditorMode::ApplyChanges()
 
 void UUVEditorMode::UpdateTriangleMaterialBasedOnBackground(bool IsBackgroundVisible)
 {
+	float TriangleOpacity;
 	// We adjust the mesh opacity depending on whether we're layered over the background or not.
 	if (IsBackgroundVisible)
 	{
-		TriangleOpacity = 0.25;
-		TriangleDepthOffset = .5;
+		TriangleOpacity = FUVEditorUXSettings::UnwrapTriangleOpacityWithBackground;
 	}
 	else
 	{
-		TriangleOpacity = 1.0;
-		TriangleDepthOffset = -.1;
+		TriangleOpacity = FUVEditorUXSettings::UnwrapTriangleOpacity;
 	}
 
 	// Modify the material of the unwrapped mesh to account for the presence/absence of the background, 
@@ -625,8 +583,8 @@ void UUVEditorMode::UpdateTriangleMaterialBasedOnBackground(bool IsBackgroundVis
 		ToolInputObjects[AssetID]->UnwrapPreview->PreviewMesh->SetMaterial(
 			0, ToolSetupUtil::GetCustomTwoSidedDepthOffsetMaterial(
 				GetToolManager(),
-				(FLinearColor)GetTriangleColorByTargetIndex(AssetID),
-				TriangleDepthOffset, 
+				FUVEditorUXSettings::GetTriangleColorByTargetIndex(AssetID),
+				FUVEditorUXSettings::UnwrapTriangleDepthOffset,
 				TriangleOpacity));
 	}
 }
diff --git a/Engine/Plugins/Experimental/UVEditor/Source/UVEditor/Public/UVEditor2DViewportClient.h b/Engine/Plugins/Experimental/UVEditor/Source/UVEditor/Public/UVEditor2DViewportClient.h
--- a/Engine/Plugins/Experimental/UVEditor/Source/UVEditor/Public/UVEditor2DViewportClient.h
+++ b/Engine/Plugins/Experimental/UVEditor/Source/UVEditor/Public/UVEditor2DViewportClient.h
@@ -50,12 +50,6 @@ protected:
 	UInputBehaviorSet* BehaviorSet;
 	UUVToolViewportButtonsAPI* ViewportButtonsAPI;
 
-	// Position to place the camera far plane relative to world z
-	float CameraFarPlaneWorldZ = -10.0;
-	// The near plane gets positioned some proportion to z = 0. We don't use a constant value because our depth offset values are percentage-based
-	// Lower proportions move the plane nearer to world z
-	float CameraNearPlaneProportionZ = 0.8;
-
 	// Note that it's generally less hassle if the unique ptr types are complete here,
 	// not forward declared, else we get compile errors if their destruction shows up
 	// anywhere in the header.
diff --git a/Engine/Plugins/Experimental/UVEditor/Source/UVEditor/Public/UVEditorMode.h b/Engine/Plugins/Experimental/UVEditor/Source/UVEditor/Public/UVEditorMode.h
--- a/Engine/Plugins/Experimental/UVEditor/Source/UVEditor/Public/UVEditorMode.h
+++ b/Engine/Plugins/Experimental/UVEditor/Source/UVEditor/Public/UVEditorMode.h
@@ -195,25 +195,6 @@ protected:
 
 	TArray<int32> PendingUVLayerIndex;
 
-	// Wireframe Display Properties
-	float TriangleOpacity = 1.0;
-	float TriangleDepthOffset = 0.5;
-	const float WireframeDepthOffset = 0.6;
-	float UnwrapBaseColorHueStart = 200;
-	float UnwrapBaseColorSaturation = 0.25;
-	float UnwrapBaseColorValue = 0.60;
-	float UnwrapWireframeHueShift = 15;
-	float UnwrapWireframeSaturation = 0.40;
-	float UnwrapWireframeValue = 0.40;
-	float UnwrapBoundaryHueShift = 30;
-	float UnwrapBoundarySaturation = 0.50;
-	float UnwrapBoundaryValue = 0.50;
-
-	FLinearColor GetTriangleColorByTargetIndex(int32 TargetIndex) const;
-	FLinearColor GetWireframeColorByTargetIndex(int32 TargetIndex) const;
-	FLinearColor GetBoundaryColorByTargetIndex(int32 TargetIndex) const;
-	FLinearColor GetSelectionColorByTargetIndex(int32 TargetIndex) const;
-
 	// Here largely for convenience to avoid having to pass it around functions.
 	UPROPERTY()
 	TObjectPtr<UWorld> LivePreviewWorld = nullptr;
diff --git a/Engine/Plugins/Experimental/UVEditor/Source/UVEditorTools/Private/UVEditorUXSettings.cpp b/Engine/Plugins/Experimental/UVEditor/Source/UVEditorTools/Private/UVEditorUXSettings.cpp
new file mode 100644
--- /dev/null
+++ b/Engine/Plugins/Experimental/UVEditor/Source/UVEditorTools/Private/UVEditorUXSettings.cpp
@@ -0,0 +1,83 @@
+// Copyright Epic Games, Inc. All Rights Reserved.
+
+#include "UVEditorUXSettings.h"
+
+const float FUVEditorUXSettings::CameraFarPlaneWorldZ(-10.0);
+
+// 2D Viewport Depth Offsets (Organized by "layers" from the camera's perspective, descending order
+const float FUVEditorUXSettings::CameraNearPlaneProportionZ(0.6); // Top layer
+const float FUVEditorUXSettings::SewLineDepthOffset(0.5f);
+const float FUVEditorUXSettings::SelectionWireframeDepthBias(0.4);
+const float FUVEditorUXSettings::SelectionTriangleDepthBias(0.3);
+const float FUVEditorUXSettings::WireframeDepthOffset(0.2);
+const float FUVEditorUXSettings::UnwrapTriangleDepthOffset(0.1);
+
+// Note that this offset can only be applied when we use our own background material
+// for a user-supplied texture, and we can't use it for a user-provided material.
+// So for consistency this should stay at zero.
+const float FUVEditorUXSettings::BackgroundQuadDepthOffset(0.0); // Bottom layer
+
+// 3D Viewport Depth Offsets
+const float FUVEditorUXSettings::LivePreviewHighlightDepthOffset(0.5);
+
+// Opacities
+const float FUVEditorUXSettings::UnwrapTriangleOpacity(1.0);
+const float FUVEditorUXSettings::UnwrapTriangleOpacityWithBackground(0.25);
+const float FUVEditorUXSettings::SelectionTriangleOpacity(1.0f);
+
+// Per Asset Shifts
+const float FUVEditorUXSettings::UnwrapBoundaryHueShift(30);
+const float FUVEditorUXSettings::UnwrapBoundarySaturation(0.50);
+const float FUVEditorUXSettings::UnwrapBoundaryValue(0.50);
+
+// Colors
+const FColor FUVEditorUXSettings::UnwrapTriangleFillColor(FColor::FromHex("#696871"));
+const FColor FUVEditorUXSettings::UnwrapTriangleWireframeColor(FColor::FromHex("#989898"));
+const FColor FUVEditorUXSettings::SelectionTriangleFillColor(FColor::FromHex("#8C7A52"));
+const FColor FUVEditorUXSettings::SelectionTriangleWireframeColor(FColor::FromHex("#DDA209"));
+const FColor FUVEditorUXSettings::SelectionHoverTriangleFillColor(FColor::FromHex("#4E719B"));
+const FColor FUVEditorUXSettings::SelectionHoverTriangleWireframeColor(FColor::FromHex("#0E86FF"));
+const FColor FUVEditorUXSettings::SewSideLeftColor(FColor::Red);
+const FColor FUVEditorUXSettings::SewSideRightColor(FColor::Green);
+const FColor FUVEditorUXSettings::XAxisColor(FColor::Red);
+const FColor FUVEditorUXSettings::YAxisColor(FColor::Green);
+const FColor FUVEditorUXSettings::GridMajorColor(FColor::FromHex("#888888"));
+const FColor FUVEditorUXSettings::GridMinorColor(FColor::FromHex("#777777"));
+
+// Thicknesses
+const float FUVEditorUXSettings::LivePreviewHighlightThickness(2.0);
+const float FUVEditorUXSettings::SelectionLineThickness(1.5);
+const float FUVEditorUXSettings::SelectionPointThickness(6);
+const float FUVEditorUXSettings::SewLineHighlightThickness(3.0f);
+const float FUVEditorUXSettings::AxisThickness(2.0);
+const float FUVEditorUXSettings::GridMajorThickness(1.0);
+
+// Grid
+const int32 FUVEditorUXSettings::GridSubdivisionsPerLevel(4);
+const int32 FUVEditorUXSettings::GridLevels(3);
+
+
+FLinearColor FUVEditorUXSettings::GetTriangleColorByTargetIndex(int32 TargetIndex)
+{
+	double GoldenAngle = 137.50776405;
+
+	FLinearColor BaseColorHSV = FLinearColor::FromSRGBColor(UnwrapTriangleFillColor).LinearRGBToHSV();
+	BaseColorHSV.R = FMath::Fmod(BaseColorHSV.R + (GoldenAngle / 2.0 * TargetIndex), 360);;
+
+	return BaseColorHSV.HSVToLinearRGB();
+}
+
+FLinearColor FUVEditorUXSettings::GetWireframeColorByTargetIndex(int32 TargetIndex)
+{
+	return FLinearColor::FromSRGBColor(UnwrapTriangleWireframeColor);;
+}
+
+FLinearColor FUVEditorUXSettings::GetBoundaryColorByTargetIndex(int32 TargetIndex)
+{
+	FLinearColor BaseColorHSV = GetTriangleColorByTargetIndex(TargetIndex).LinearRGBToHSV();
+	FLinearColor BoundaryColorHSV = BaseColorHSV;
+	BoundaryColorHSV.R = FMath::Fmod((BoundaryColorHSV.R + UnwrapBoundaryHueShift), 360);
+	BoundaryColorHSV.G = UnwrapBoundarySaturation;
+	BoundaryColorHSV.B = UnwrapBoundaryValue;
+	return BoundaryColorHSV.HSVToLinearRGB();
+}
diff --git a/Engine/Plugins/Experimental/UVEditor/Source/UVEditorTools/Private/UVSeamSewAction.cpp b/Engine/Plugins/Experimental/UVEditor/Source/UVEditorTools/Private/UVSeamSewAction.cpp
--- a/Engine/Plugins/Experimental/UVEditor/Source/UVEditorTools/Private/UVSeamSewAction.cpp
+++ b/Engine/Plugins/Experimental/UVEditor/Source/UVEditorTools/Private/UVSeamSewAction.cpp
@@ -11,6 +11,8 @@
 #include "PreviewMesh.h"
 #include "ToolSetupUtil.h"
 #include "UVToolContextObjects.h"
+#include "UVEditorUXSettings.h"
+
 
 #define LOCTEXT_NAMESPACE "UUVSeamSewAction"
 
@@ -140,14 +142,14 @@ void UUVSeamSewAction::UpdateSewEdgePreviewLines()
 			SewEdgePairingLineSet->AddLine(
 				MeshTransform.TransformPosition(Vert1),
 				MeshTransform.TransformPosition(Vert2),
-				FColor::Red, 5, 3);
+				FUVEditorUXSettings::SewSideLeftColor, FUVEditorUXSettings::SewLineHighlightThickness, FUVEditorUXSettings::SewLineDepthOffset);
 
 			Vert1 = UnwrapMesh->GetVertex(SewPair.B[0]);
 			Vert2 = UnwrapMesh->GetVertex(SewPair.B[1]);
 			SewEdgePairingLineSet->AddLine(
 				MeshTransform.TransformPosition(Vert1),
 				MeshTransform.TransformPosition(Vert2),
-				FColor::Green, 5, 3);
+				FUVEditorUXSettings::SewSideRightColor, FUVEditorUXSettings::SewLineHighlightThickness, FUVEditorUXSettings::SewLineDepthOffset);
 
 		}
 	}
diff --git a/Engine/Plugins/Experimental/UVEditor/Source/UVEditorTools/Private/UVSelectTool.cpp b/Engine/Plugins/Experimental/UVEditor/Source/UVEditorTools/Private/UVSelectTool.cpp
--- a/Engine/Plugins/Experimental/UVEditor/Source/UVEditorTools/Private/UVSelectTool.cpp
+++ b/Engine/Plugins/Experimental/UVEditor/Source/UVEditorTools/Private/UVSelectTool.cpp
@@ -19,6 +19,7 @@
 #include "Selection/MeshSelectionMechanic.h"
 #include "Selection/DynamicMeshSelection.h"
 #include "ToolSetupUtil.h"
+#include "UVEditorUXSettings.h"
 
 #include "UVSeamSewAction.h"
 #include "UVIslandConformalUnwrapAction.h"
@@ -323,8 +324,17 @@ void UUVSelectTool::Setup()
 	SelectionMechanic->Setup(this);
 	SelectionMechanic->SetWorld(Targets[0]->UnwrapPreview->GetWorld());
 	SelectionMechanic->OnSelectionChanged.AddUObject(this, &UUVSelectTool::OnSelectionChanged);
-
-
+	FMeshSelectionMechanicStyle SelectionStyle;
+	SelectionStyle.TriangleColor = FUVEditorUXSettings::SelectionTriangleFillColor;
+	SelectionStyle.LineColor = FUVEditorUXSettings::SelectionTriangleWireframeColor;
+	SelectionStyle.PointColor = FUVEditorUXSettings::SelectionTriangleWireframeColor;
+	SelectionStyle.TriangleOpacity = FUVEditorUXSettings::SelectionTriangleOpacity;
+	SelectionStyle.LineThickness = FUVEditorUXSettings::SelectionLineThickness;
+	SelectionStyle.PointThickness = FUVEditorUXSettings::SelectionPointThickness;
+	SelectionStyle.LineAndPointDepthBias = FUVEditorUXSettings::SelectionWireframeDepthBias;
+	SelectionStyle.TriangleDepthBias = FUVEditorUXSettings::SelectionTriangleDepthBias;
+	SelectionMechanic->SetVisualizationStyle(SelectionStyle);
+		
 	// Make it so that our selection mechanic creates undo/redo transactions that go to a selection
 	// change router, which we use to route to the current selection mechanic on each tool invocation.
 	ChangeRouter = ContextStore->FindContext<UUVSelectToolChangeRouter>();
@@ -646,14 +656,6 @@ void UUVSelectTool::OnSelectionChanged()
 		}
 		check(SelectionTargetIndex >= 0);
 
-		UContextObjectStore* ContextStore = GetToolManager()->GetContextObjectStore();
-		UUVVisualStyleAPI* VisualStyleAPI = ContextStore->FindContext<UUVVisualStyleAPI>();
-		if (VisualStyleAPI)
-		{
-			FColor SelectionColor = VisualStyleAPI->GetSelectionColorForAsset(SelectionTargetIndex).ToFColor(true);
-			SelectionMechanic->ChangeSelectionColor(SelectionColor, 0.3f, SelectionColor, SelectionColor);
-		}
-
 		// Note the selected vids
 		TSet<int32> VidSet;
 		TSet<int32> TidSet;
@@ -778,14 +780,6 @@ void UUVSelectTool::UpdateLivePreviewLines()
 		FTransform MeshTransform = Targets[SelectionTargetIndex]->AppliedPreview->PreviewMesh->GetTransform();
 		const FDynamicMesh3* LivePreviewMesh = Targets[SelectionTargetIndex]->AppliedCanonical.Get();
 
-		UContextObjectStore* ContextStore = GetToolManager()->GetContextObjectStore();
-		UUVVisualStyleAPI* VisualStyleAPI = ContextStore->FindContext<UUVVisualStyleAPI>();
-		FColor SelectionColor = FColor::Yellow;
-		if (VisualStyleAPI)
-		{
-			SelectionColor = VisualStyleAPI->GetSelectionColorForAsset(SelectionTargetIndex).ToFColor(true);
-		}
-
 		for (int32 Eid : LivePreviewBoundaryEids)
 		{
 			FVector3d Vert1, Vert2;
@@ -794,7 +788,9 @@ void UUVSelectTool::UpdateLivePreviewLines()
 			LivePreviewLineSet->AddLine(
 				MeshTransform.TransformPosition(Vert1), 
 				MeshTransform.TransformPosition(Vert2), 
-				SelectionColor, LivePreviewHighlightThickness, LivePreviewHighlightDepthOffset);
+				FUVEditorUXSettings::SelectionTriangleWireframeColor,
+				FUVEditorUXSettings::LivePreviewHighlightThickness,
+				FUVEditorUXSettings::LivePreviewHighlightDepthOffset);
 		}	
 	}
 }
diff --git a/Engine/Plugins/Experimental/UVEditor/Source/UVEditorTools/Public/UVEditorUXSettings.h b/Engine/Plugins/Experimental/UVEditor/Source/UVEditorTools/Public/UVEditorUXSettings.h
new file mode 100644
--- /dev/null
+++ b/Engine/Plugins/Experimental/UVEditor/Source/UVEditorTools/Public/UVEditorUXSettings.h
@@ -0,0 +1,79 @@
+// Copyright Epic Games, Inc. All Rights Reserved.
+
+#pragma once
+
+class UVEDITORTOOLS_API FUVEditorUXSettings
+{
+public:
+	// The following are UV Editor specific style items that don't, strictly, matter to the SlateStyleSet,
+	// but seemed appropriate to place here.
+
+	// General Display Properties
+
+	// Position to place the 2D camera far plane relative to world z
+	static const float CameraFarPlaneWorldZ;;
+
+	// The near plane gets positioned some proportion to z = 0. We don't use a constant value because our depth offset values are percentage-based
+	// Lower proportion move the plane nearer to world z
+	// Note: This serves as an upper bound for all other depth offsets - higher than this value risks being clipped
+	static const float CameraNearPlaneProportionZ;;
+
+
+	// 2D Unwrap Display Properties
+	static const float UnwrapTriangleOpacity;
+	static const float UnwrapTriangleDepthOffset;
+	static const float UnwrapTriangleOpacityWithBackground;
+
+	static const float WireframeDepthOffset;
+	static const float UnwrapBoundaryHueShift;
+	static const float UnwrapBoundarySaturation;
+	static const float UnwrapBoundaryValue;
+	static const FColor UnwrapTriangleFillColor;
+	static const FColor UnwrapTriangleWireframeColor;
+
+	static FLinearColor GetTriangleColorByTargetIndex(int32 TargetIndex);
+	static FLinearColor GetWireframeColorByTargetIndex(int32 TargetIndex);
+	static FLinearColor GetBoundaryColorByTargetIndex(int32 TargetIndex);
+
+	// Selection Highlighting Properties
+	static const float SelectionTriangleOpacity;
+	static const FColor SelectionTriangleFillColor;
+	static const FColor SelectionTriangleWireframeColor;
+
+	// Todo: Not used yet, for the upcoming pre-selection hover feature
+	static const FColor SelectionHoverTriangleFillColor;
+	static const FColor SelectionHoverTriangleWireframeColor;
+
+	static const float LivePreviewHighlightThickness;
+	static const float LivePreviewHighlightDepthOffset;
+
+	static const float SelectionLineThickness;
+	static const float SelectionPointThickness;
+	static const float SelectionWireframeDepthBias;
+	static const float SelectionTriangleDepthBias;
+
+	// Sew Action styling
+	static const float SewLineHighlightThickness;
+	static const float SewLineDepthOffset;
+	static const FColor SewSideLeftColor;
+	static const FColor SewSideRightColor;
+
+	// Grid 
+	static const float AxisThickness;
+	static const float GridMajorThickness;
+	static const FColor XAxisColor;
+	static const FColor YAxisColor;
+	static const FColor GridMajorColor;
+	static const FColor GridMinorColor;
+	static const int32 GridSubdivisionsPerLevel;
+	static const int32 GridLevels;
+
+
+	// Background
+	static const float BackgroundQuadDepthOffset;
+
+private:
+
+	FUVEditorUXSettings();
+	~FUVEditorUXSettings();
+};
\ No newline at end of file
diff --git a/Engine/Plugins/Experimental/UVEditor/Source/UVEditorTools/Public/UVSelectTool.h b/Engine/Plugins/Experimental/UVEditor/Source/UVEditorTools/Public/UVSelectTool.h
--- a/Engine/Plugins/Experimental/UVEditor/Source/UVEditorTools/Public/UVSelectTool.h
+++ b/Engine/Plugins/Experimental/UVEditor/Source/UVEditorTools/Public/UVSelectTool.h
@@ -198,8 +198,7 @@ protected:
 	void ApplyAction(ESelectToolAction ActionType);
 	void ApplySplit();
 
-	const float LivePreviewHighlightThickness = 2.0;
-	const float LivePreviewHighlightDepthOffset = 0.5;
+
 };
 
 /**
diff --git a/Engine/Plugins/Experimental/UVEditor/Source/UVEditorTools/Public/UVToolContextObjects.h b/Engine/Plugins/Experimental/UVEditor/Source/UVEditorTools/Public/UVToolContextObjects.h
--- a/Engine/Plugins/Experimental/UVEditor/Source/UVEditorTools/Public/UVToolContextObjects.h
+++ b/Engine/Plugins/Experimental/UVEditor/Source/UVEditorTools/Public/UVToolContextObjects.h
@@ -247,28 +247,6 @@ public:
 
 };
 
-
-/**
-  Allows tools and mode to access color and styling information per asset
-*/
-UCLASS()
-class UVEDITORTOOLS_API UUVVisualStyleAPI : public UUVToolContextObject
-{
-	GENERATED_BODY()
-public:
-	FLinearColor GetSelectionColorForAsset(int32 AssetID)
-	{
-		if (GetSelectionColorForAssetFunc)
-		{
-			return GetSelectionColorForAssetFunc(AssetID);
-		}
-		return FLinearColor();
-	}
-
-	TUniqueFunction<FLinearColor(int32)> GetSelectionColorForAssetFunc;
-};
-
-
 /** 
  * Stores AABB trees for UV input object unwrap canonical or applied canonical meshes.
  * Binds to the input objects' OnCanonicalModified delegate to automatically update 
diff --git a/Engine/Plugins/Runtime/MeshModelingToolset/Source/ModelingComponents/Private/Selection/MeshSelectionMechanic.cpp b/Engine/Plugins/Runtime/MeshModelingToolset/Source/ModelingComponents/Private/Selection/MeshSelectionMechanic.cpp
--- a/Engine/Plugins/Runtime/MeshModelingToolset/Source/ModelingComponents/Private/Selection/MeshSelectionMechanic.cpp
+++ b/Engine/Plugins/Runtime/MeshModelingToolset/Source/ModelingComponents/Private/Selection/MeshSelectionMechanic.cpp
@@ -296,7 +296,8 @@ void UMeshSelectionMechanic::Setup(UInteractiveTool* ParentToolIn)
 	ClearCurrentSelection();
 
 	TriangleSet = NewObject<UTriangleSetComponent>();
-	TriangleSetMaterial = ToolSetupUtil::GetCustomTwoSidedDepthOffsetMaterial(GetParentTool()->GetToolManager(), TriangleColor, DepthBias, TriangleOpacity);
+	TriangleSetMaterial = ToolSetupUtil::GetCustomTwoSidedDepthOffsetMaterial(GetParentTool()->GetToolManager(),
+		VisualizationStyle.TriangleColor, VisualizationStyle.TriangleDepthBias, VisualizationStyle.TriangleOpacity);
 	
 	LineSet = NewObject<ULineSetComponent>();
 	LineSet->SetLineMaterial(ToolSetupUtil::GetDefaultLineComponentMaterial(
@@ -387,16 +388,14 @@ const FDynamicMeshSelection& UMeshSelectionMechanic::GetCurrentSelection() const
 	return CurrentSelection;
 }
 
-void UMeshSelectionMechanic::ChangeSelectionColor(const FColor& TriangleColorIn, float TriangleOpacityIn, const FColor& LineColorIn, const FColor& PointColorIn)
+void UMeshSelectionMechanic::SetVisualizationStyle(const FMeshSelectionMechanicStyle& StyleIn)
 {
-	TriangleColor = TriangleColorIn;
-	TriangleOpacity = TriangleOpacityIn;
-	LineColor = LineColorIn;
-	PointColor = PointColorIn;
+	VisualizationStyle = StyleIn;
 
 	if (TriangleSetMaterial) {
-		TriangleSetMaterial->SetVectorParameterValue(TEXT("Color"), TriangleColorIn);
-		TriangleSetMaterial->SetScalarParameterValue(TEXT("Opacity"), TriangleOpacityIn);
+		TriangleSetMaterial->SetVectorParameterValue(TEXT("Color"), VisualizationStyle.TriangleColor);
+		TriangleSetMaterial->SetScalarParameterValue(TEXT("Opacity"), VisualizationStyle.TriangleOpacity);
+		TriangleSetMaterial->SetScalarParameterValue(TEXT("PercentDepthOffset"), VisualizationStyle.TriangleDepthBias);
 	}
 }
 
@@ -481,12 +480,12 @@ void UMeshSelectionMechanic::RebuildDrawnElements(const FTransform& StartTransfo
 			{
 				Points[i] = TransformToApply(CurrentSelection.Mesh->GetVertex(Vids[i]));
 			}
-			TriangleSet->AddTriangle(Points[0], Points[1], Points[2], FVector(0, 0, 1), TriangleColor, TriangleSetMaterial);
+			TriangleSet->AddTriangle(Points[0], Points[1], Points[2], FVector(0, 0, 1), VisualizationStyle.TriangleColor, TriangleSetMaterial);
 			for (int i = 0; i < 3; ++i)
 			{
 				int NextIndex = (i + 1) % 3;
 				LineSet->AddLine(Points[i], Points[NextIndex],
-					LineColor, LineThickness, DepthBias);
+					VisualizationStyle.LineColor, VisualizationStyle.LineThickness, VisualizationStyle.LineAndPointDepthBias);
 			}
 		}
 	}
@@ -501,7 +500,7 @@ void UMeshSelectionMechanic::RebuildDrawnElements(const FTransform& StartTransfo
 			LineSet->AddLine(
 				TransformToApply(CurrentSelection.Mesh->GetVertex(EdgeVids.A)),
 				TransformToApply(CurrentSelection.Mesh->GetVertex(EdgeVids.B)),
-					LineColor, LineThickness, DepthBias);
+				VisualizationStyle.LineColor, VisualizationStyle.LineThickness, VisualizationStyle.LineAndPointDepthBias);
 		}
 	}
 	else if (CurrentSelection.Type == FDynamicMeshSelection::EType::Vertex)
@@ -512,8 +511,8 @@ void UMeshSelectionMechanic::RebuildDrawnElements(const FTransform& StartTransfo
 		for (int32 Vid : CurrentSelection.SelectedIDs)
 		{
 			FRenderablePoint PointToRender(TransformToApply(CurrentSelection.Mesh->GetVertex(Vid)),
-				                           PointColor,
-				                           PointThickness);
+				VisualizationStyle.PointColor,
+				VisualizationStyle.PointThickness);
 			PointSet->AddPoint(PointToRender);
 		}
 	}
diff --git a/Engine/Plugins/Runtime/MeshModelingToolset/Source/ModelingComponents/Public/Selection/MeshSelectionMechanic.h b/Engine/Plugins/Runtime/MeshModelingToolset/Source/ModelingComponents/Public/Selection/MeshSelectionMechanic.h
--- a/Engine/Plugins/Runtime/MeshModelingToolset/Source/ModelingComponents/Public/Selection/MeshSelectionMechanic.h
+++ b/Engine/Plugins/Runtime/MeshModelingToolset/Source/ModelingComponents/Public/Selection/MeshSelectionMechanic.h
@@ -44,6 +44,18 @@ class MODELINGCOMPONENTS_API UMeshSelectionMechanicProperties : public UInteract
 public:
 };
 
+struct MODELINGCOMPONENTS_API FMeshSelectionMechanicStyle
+{
+	FColor TriangleColor = FColor::Yellow;
+	FColor LineColor = FColor::Yellow;
+	FColor PointColor = FColor::Yellow;
+	float TriangleOpacity = 0.3;
+	float LineThickness = 1.5;
+	float PointThickness = 6;
+	float LineAndPointDepthBias = 0.7;
+	float TriangleDepthBias = 0.65;
+};
+
 /**
  * Mechanic for selecting elements of a dynamic mesh.
  * 
@@ -80,7 +92,7 @@ public:
 
 	virtual const FDynamicMeshSelection& GetCurrentSelection() const;
 	void ChangeSelectionMode(const EMeshSelectionMechanicMode& TargetMode);
-	void ChangeSelectionColor(const FColor& TriangleColorIn, float TriangleOpacityIn, const FColor& LineColorIn, const FColor& PointColorIn);
+	void SetVisualizationStyle(const FMeshSelectionMechanicStyle& StyleIn);
 	virtual void SetSelection(const FDynamicMeshSelection& Selection, bool bBroadcast = false, bool bEmitChange = false);
 
 	virtual void Render(IToolsContextRenderAPI* RenderAPI) override;
@@ -146,20 +158,13 @@ protected:
 	int32 CurrentSelectionIndex = IndexConstants::InvalidID;
 	FViewCameraState CameraState;
 
-	FColor TriangleColor = FColor::Yellow;
-	FColor LineColor = FColor::Yellow;
-	FColor PointColor = FColor::Yellow;
-	float TriangleOpacity = 0.3;
+	FMeshSelectionMechanicStyle VisualizationStyle;
 
 	bool bShiftToggle = false;
 	bool bCtrlToggle = false;
 	static const int32 ShiftModifierID = 1;
 	static const int32 CtrlModifierID = 2;
 
-	float LineThickness = 3;
-	float PointThickness = 6;
-	float DepthBias = 0.7;
-
 	FVector3d CurrentSelectionCentroid;
 	int32 CentroidTimestamp = -1;
 	void UpdateCentroid();
