UVEditor: Fix issue where selection triangles would disappear in and out of view as the camera moved. This fix adjusts the TranslucentSortPriority of the selection TriangleSet to keep it above the 2D wireframe mesh below.

#rnx

diff --git a/Engine/Plugins/Experimental/UVEditor/Source/UVEditorTools/Private/UVEditorUXSettings.cpp b/Engine/Plugins/Experimental/UVEditor/Source/UVEditorTools/Private/UVEditorUXSettings.cpp
--- a/Engine/Plugins/Experimental/UVEditor/Source/UVEditorTools/Private/UVEditorUXSettings.cpp
+++ b/Engine/Plugins/Experimental/UVEditor/Source/UVEditorTools/Private/UVEditorUXSettings.cpp
@@ -5,16 +5,20 @@
 const float FUVEditorUXSettings::CameraFarPlaneWorldZ(-10.0);
 
 // 2D Viewport Depth Offsets (Organized by "layers" from the camera's perspective, descending order
-const float FUVEditorUXSettings::CameraNearPlaneProportionZ(0.6); // Top layer
-const float FUVEditorUXSettings::SewLineDepthOffset(0.5f);
-const float FUVEditorUXSettings::SelectionWireframeDepthBias(0.4);
-const float FUVEditorUXSettings::SelectionTriangleDepthBias(0.3);
-const float FUVEditorUXSettings::WireframeDepthOffset(0.2);
-const float FUVEditorUXSettings::UnwrapTriangleDepthOffset(0.1);
-
-// Note that this offset can only be applied when we use our own background material
+// Note: While these are floating point values, they represent percentages and should be separated
+// by at least integer amounts, as they serve double duty in certain cases for translucent primitive
+// sorting order.
+const float FUVEditorUXSettings::CameraNearPlaneProportionZ(8); // Top layer
+const float FUVEditorUXSettings::SewLineDepthOffset(5.0f);
+const float FUVEditorUXSettings::SelectionWireframeDepthBias(4.0);
+const float FUVEditorUXSettings::SelectionTriangleDepthBias(3.0);
+const float FUVEditorUXSettings::WireframeDepthOffset(2.0);
+const float FUVEditorUXSettings::UnwrapTriangleDepthOffset(1.0);
+
+// Note: that this offset can only be applied when we use our own background material
 // for a user-supplied texture, and we can't use it for a user-provided material.
 // So for consistency this should stay at zero.
+
 const float FUVEditorUXSettings::BackgroundQuadDepthOffset(0.0); // Bottom layer
 
 // 3D Viewport Depth Offsets
diff --git a/Engine/Plugins/Runtime/MeshModelingToolset/Source/ModelingComponents/Private/Selection/MeshSelectionMechanic.cpp b/Engine/Plugins/Runtime/MeshModelingToolset/Source/ModelingComponents/Private/Selection/MeshSelectionMechanic.cpp
--- a/Engine/Plugins/Runtime/MeshModelingToolset/Source/ModelingComponents/Private/Selection/MeshSelectionMechanic.cpp
+++ b/Engine/Plugins/Runtime/MeshModelingToolset/Source/ModelingComponents/Private/Selection/MeshSelectionMechanic.cpp
@@ -296,6 +296,9 @@ void UMeshSelectionMechanic::Setup(UInteractiveTool* ParentToolIn)
 	ClearCurrentSelection();
 
 	TriangleSet = NewObject<UTriangleSetComponent>();
+	// We are setting the TranslucencySortPriority here to handle the UV editor's use case in 2D
+	// where multiple translucent layers are drawn on top of each other but still need depth sorting.
+	TriangleSet->TranslucencySortPriority = VisualizationStyle.TriangleDepthBias;	
 	TriangleSetMaterial = ToolSetupUtil::GetCustomTwoSidedDepthOffsetMaterial(GetParentTool()->GetToolManager(),
 		VisualizationStyle.TriangleColor, VisualizationStyle.TriangleDepthBias, VisualizationStyle.TriangleOpacity);
 	
@@ -392,6 +395,12 @@ void UMeshSelectionMechanic::SetVisualizationStyle(const FMeshSelectionMechanicS
 {
 	VisualizationStyle = StyleIn;
 
+	if (TriangleSet) {
+		// We are setting the TranslucencySortPriority here to handle the UV editor's use case in 2D
+		// where multiple translucent layers are drawn on top of each other but still need depth sorting.
+		TriangleSet->TranslucencySortPriority = VisualizationStyle.TriangleDepthBias;
+	}
+
 	if (TriangleSetMaterial) {
 		TriangleSetMaterial->SetVectorParameterValue(TEXT("Color"), VisualizationStyle.TriangleColor);
 		TriangleSetMaterial->SetScalarParameterValue(TEXT("Opacity"), VisualizationStyle.TriangleOpacity);
