ModelingTools: Change PointSetComponent to match LineSetComponent to use per element (point) percentage-based DepthBias support. Updates users of PointSetComponent to reflect this change.

#rnx

diff --git a/Engine/Plugins/Experimental/UVEditor/Source/UVEditorTools/Private/UVEditorSeamTool.cpp b/Engine/Plugins/Experimental/UVEditor/Source/UVEditorTools/Private/UVEditorSeamTool.cpp
--- a/Engine/Plugins/Experimental/UVEditor/Source/UVEditorTools/Private/UVEditorSeamTool.cpp
+++ b/Engine/Plugins/Experimental/UVEditor/Source/UVEditorTools/Private/UVEditorSeamTool.cpp
@@ -83,17 +83,17 @@ namespace UVEditorSeamToolLocals
 
 	void AddDisplayedPoints(UUVEditorToolMeshInput* InputObject,
 		UPointSetComponent* UnwrapPointSet, UPointSetComponent* AppliedPointSet,
-		int32 AppliedVid, const TArray<int32>& UnwrapVids, const FColor& Color)
+		int32 AppliedVid, const TArray<int32>& UnwrapVids, const FColor& Color, float DepthBias)
 	{
 		TRACE_CPUPROFILER_EVENT_SCOPE(AddDisplayedPoints);
 
 		FTransform AppliedTransform = InputObject->AppliedPreview->PreviewMesh->GetTransform();
 		AppliedPointSet->AddPoint(FRenderablePoint(
-			InputObject->AppliedCanonical->GetVertex(AppliedVid), Color, PointSize));
+			InputObject->AppliedCanonical->GetVertex(AppliedVid), Color, PointSize, DepthBias));
 		for (int32 UnwrapVid : UnwrapVids)
 		{
 			UnwrapPointSet->AddPoint(FRenderablePoint(
-				InputObject->UnwrapCanonical->GetVertex(UnwrapVid), Color, PointSize));
+				InputObject->UnwrapCanonical->GetVertex(UnwrapVid), Color, PointSize, DepthBias));
 		}
 	}
 
@@ -536,11 +536,11 @@ void UUVEditorSeamTool::ReconstructLockedPathVisualization()
 
 		TArray<int32> UnwrapVids;
 		ClickedTarget->AppliedVidToUnwrapVids(LastLockedAppliedVid, UnwrapVids);
-		AddDisplayedPoints(ClickedTarget, UnwrapPointSet, AppliedPointSet, LastLockedAppliedVid, UnwrapVids, LockedColor);
+		AddDisplayedPoints(ClickedTarget, UnwrapPointSet, AppliedPointSet, LastLockedAppliedVid, UnwrapVids, LockedColor, LockedDepthBias);
 
 		UnwrapVids.Reset();
 		ClickedTarget->AppliedVidToUnwrapVids(SeamStartAppliedVid, UnwrapVids);
-		AddDisplayedPoints(ClickedTarget, UnwrapStartPoints, AppliedStartPoints, SeamStartAppliedVid, UnwrapVids, LockedColor);
+		AddDisplayedPoints(ClickedTarget, UnwrapStartPoints, AppliedStartPoints, SeamStartAppliedVid, UnwrapVids, LockedColor, LockedDepthBias);
 	}
 }
 
@@ -630,7 +630,7 @@ void UUVEditorSeamTool::UpdateHover()
 	{
 		// Just draw the point and finish
 		AddDisplayedPoints(Targets[HoverMeshIndex], UnwrapPointSet, AppliedPointSet,
-			AppliedVid, UnwrapVids, HoverColor);
+			AppliedVid, UnwrapVids, HoverColor, HoverDepthBias);
 		return;
 	}
 
@@ -672,7 +672,7 @@ void UUVEditorSeamTool::UpdateHover()
 		// there because that turns out to be confusing, esp when we have overlapping islands and
 		// the point is on another island.
 		AddDisplayedPoints(Targets[HoverMeshIndex], UnwrapPointSet, AppliedPointSet,
-			AppliedVid, UnwrapVids, HoverColor);
+			AppliedVid, UnwrapVids, HoverColor, HoverDepthBias);
 	}
 }
 
@@ -743,7 +743,7 @@ void UUVEditorSeamTool::ApplyClick()
 
 		UPointSetComponent* UnwrapPointSet = UnwrapGeometry->FindPointSet(StartPointSetID);
 		UPointSetComponent* AppliedPointSet = LivePreviewGeometry->FindPointSet(StartPointSetID);
-		AddDisplayedPoints(ClickedTarget, UnwrapPointSet, AppliedPointSet, AppliedVid, UnwrapVids, LockedColor);
+		AddDisplayedPoints(ClickedTarget, UnwrapPointSet, AppliedPointSet, AppliedVid, UnwrapVids, LockedColor, LockedDepthBias);
 
 		LockedPath.Add(AppliedVid);
 		EmitChangeAPI->EmitToolDependentChange(this, 
@@ -813,7 +813,7 @@ void UUVEditorSeamTool::ApplyClick()
 		bClickWasInUnwrap, LockedColor, LockedThickness, LockedDepthBias);
 	UnwrapPointSet->Clear();
 	AppliedPointSet->Clear();
-	AddDisplayedPoints(ClickedTarget, UnwrapPointSet, AppliedPointSet, AppliedVid, UnwrapVids, LockedColor);
+	AddDisplayedPoints(ClickedTarget, UnwrapPointSet, AppliedPointSet, AppliedVid, UnwrapVids, LockedColor, LockedDepthBias);
 	
 	EmitChangeAPI->EmitToolDependentChange(this, 
 		MakeUnique<FPathChange>(*AppliedVidPath, false, ClickedMeshIndex, true),
diff --git a/Engine/Plugins/Runtime/MeshModelingToolset/Source/ModelingComponents/Private/Drawing/PointSetComponent.cpp b/Engine/Plugins/Runtime/MeshModelingToolset/Source/ModelingComponents/Private/Drawing/PointSetComponent.cpp
--- a/Engine/Plugins/Runtime/MeshModelingToolset/Source/ModelingComponents/Private/Drawing/PointSetComponent.cpp
+++ b/Engine/Plugins/Runtime/MeshModelingToolset/Source/ModelingComponents/Private/Drawing/PointSetComponent.cpp
@@ -107,7 +107,7 @@ public:
 				// have to convert it back to linear. The ToFColor(false) call just scales back into 0-255 space.
 				FColor color = FLinearColor::FromSRGBColor(Point.Color).ToFColor(false);
 
-				const FVector2D UV(Point.Size, 0.0f);
+				const FVector2D UV(Point.Size, Point.DepthBias);
 				for (int j = 0; j < 4; ++j)
 				{
 					VertexBuffers.PositionVertexBuffer.VertexPosition(UseVertexBufferIndex + j) = Point.Position;
diff --git a/Engine/Plugins/Runtime/MeshModelingToolset/Source/ModelingComponents/Private/Mechanics/CurveControlPointsMechanic.cpp b/Engine/Plugins/Runtime/MeshModelingToolset/Source/ModelingComponents/Private/Mechanics/CurveControlPointsMechanic.cpp
--- a/Engine/Plugins/Runtime/MeshModelingToolset/Source/ModelingComponents/Private/Mechanics/CurveControlPointsMechanic.cpp
+++ b/Engine/Plugins/Runtime/MeshModelingToolset/Source/ModelingComponents/Private/Mechanics/CurveControlPointsMechanic.cpp
@@ -281,7 +281,7 @@ int32 UCurveControlPointsMechanic::InsertPointAt(int32 SequencePosition, const F
 	// Add the point
 	int32 NewPointID = ControlPoints.InsertPointAt(SequencePosition, NewPointCoordinates, KnownPointID);
 	GeometrySet.AddPoint(NewPointID, NewPointCoordinates);
-	FRenderablePoint RenderablePoint((FVector)NewPointCoordinates, CurrentPointsColor, PointsSize);
+	FRenderablePoint RenderablePoint((FVector)NewPointCoordinates, CurrentPointsColor, PointsSize, DepthBias);
 	DrawnControlPoints->InsertPoint(NewPointID, RenderablePoint);
 
 	// See if we need to add some segments
@@ -920,7 +920,7 @@ bool UCurveControlPointsMechanic::OnUpdateHover(const FInputDeviceRay& DevicePos
 	if (!bInteractiveInitializationMode && bInsertPointToggle && GeometrySet.FindNearestCurveToRay((FRay3d)DevicePos.WorldRay, Nearest, GeometrySetToleranceTest))
 	{
 		ClearHover();
-		FRenderablePoint RenderablePoint((FVector)Nearest.NearestGeoPoint, PreviewColor, PointsSize);
+		FRenderablePoint RenderablePoint((FVector)Nearest.NearestGeoPoint, PreviewColor, PointsSize, DepthBias);
 		PreviewPoint->InsertPoint(0, RenderablePoint);
 	}
 
@@ -970,7 +970,7 @@ bool UCurveControlPointsMechanic::OnUpdateHover(const FInputDeviceRay& DevicePos
 
 			// Redraw preview
 			ClearHover();
-			FRenderablePoint RenderablePoint((FVector)HitPoint, PreviewColor, PointsSize);
+			FRenderablePoint RenderablePoint((FVector)HitPoint, PreviewColor, PointsSize, DepthBias);
 			PreviewPoint->InsertPoint(0, RenderablePoint);
 
 			if (ControlPoints.Num() > 0)
diff --git a/Engine/Plugins/Runtime/MeshModelingToolset/Source/ModelingComponents/Private/Mechanics/SpaceCurveDeformationMechanic.cpp b/Engine/Plugins/Runtime/MeshModelingToolset/Source/ModelingComponents/Private/Mechanics/SpaceCurveDeformationMechanic.cpp
--- a/Engine/Plugins/Runtime/MeshModelingToolset/Source/ModelingComponents/Private/Mechanics/SpaceCurveDeformationMechanic.cpp
+++ b/Engine/Plugins/Runtime/MeshModelingToolset/Source/ModelingComponents/Private/Mechanics/SpaceCurveDeformationMechanic.cpp
@@ -852,7 +852,7 @@ void USpaceCurveDeformationMechanic::UpdateRenderGeometry()
 			FFrame3d Point = CurveSource->GetPoint(k);
 			bool bSelected = SelectedPointIDs.Contains(k);
 			RenderPoints->AddPoint(FRenderablePoint((FVector)Point.Origin, 
-				(bSelected) ? SelectedColor : CurrentPointsColor, PointsSize));
+				(bSelected) ? SelectedColor : CurrentPointsColor, PointsSize, DepthBias));
 		}
 
 		int32 StopIndex = CurveSource->IsLoop() ? NumPoints : (NumPoints-1);
diff --git a/Engine/Plugins/Runtime/MeshModelingToolset/Source/ModelingComponents/Private/Selection/MeshSelectionMechanic.cpp b/Engine/Plugins/Runtime/MeshModelingToolset/Source/ModelingComponents/Private/Selection/MeshSelectionMechanic.cpp
--- a/Engine/Plugins/Runtime/MeshModelingToolset/Source/ModelingComponents/Private/Selection/MeshSelectionMechanic.cpp
+++ b/Engine/Plugins/Runtime/MeshModelingToolset/Source/ModelingComponents/Private/Selection/MeshSelectionMechanic.cpp
@@ -521,7 +521,8 @@ void UMeshSelectionMechanic::RebuildDrawnElements(const FTransform& StartTransfo
 		{
 			FRenderablePoint PointToRender(TransformToApply(CurrentSelection.Mesh->GetVertex(Vid)),
 				VisualizationStyle.PointColor,
-				VisualizationStyle.PointThickness);
+				VisualizationStyle.PointThickness,
+				VisualizationStyle.LineAndPointDepthBias);
 			PointSet->AddPoint(PointToRender);
 		}
 	}
diff --git a/Engine/Plugins/Runtime/MeshModelingToolset/Source/ModelingComponents/Public/Drawing/PointSetComponent.h b/Engine/Plugins/Runtime/MeshModelingToolset/Source/ModelingComponents/Public/Drawing/PointSetComponent.h
--- a/Engine/Plugins/Runtime/MeshModelingToolset/Source/ModelingComponents/Public/Drawing/PointSetComponent.h
+++ b/Engine/Plugins/Runtime/MeshModelingToolset/Source/ModelingComponents/Public/Drawing/PointSetComponent.h
@@ -16,17 +16,20 @@ struct FRenderablePoint
 		: Position(ForceInitToZero)
 		, Color(ForceInitToZero)
 		, Size(0.f)
+		, DepthBias(0.f)
 	{}
 
-	FRenderablePoint(const FVector& InPosition, const FColor& InColor, const float InSize)
+	FRenderablePoint(const FVector& InPosition, const FColor& InColor, const float InSize, const float InDepthBias = 0.0f)
 		: Position(InPosition)
 		, Color(InColor)
 		, Size(InSize)
+		, DepthBias(InDepthBias)
 	{}
 
 	FVector Position;
 	FColor Color;
 	float Size;
+	float DepthBias;
 };
 
 
