UVEditor: add ability to split bowtie verts. Fix some edge splits occasionally creating a bowtie vert on adjacent non-selected edge. Also fix some undo/redo issues especially regarding sew edges preview. Made all undo/redo selection change events broadcast for now.

#rnx

diff --git a/Engine/Plugins/Experimental/UVEditor/Source/UVEditor/Private/UVEditorCommands.cpp b/Engine/Plugins/Experimental/UVEditor/Source/UVEditor/Private/UVEditorCommands.cpp
--- a/Engine/Plugins/Experimental/UVEditor/Source/UVEditor/Private/UVEditorCommands.cpp
+++ b/Engine/Plugins/Experimental/UVEditor/Source/UVEditor/Private/UVEditorCommands.cpp
@@ -33,9 +33,11 @@ void FUVEditorCommands::RegisterCommands()
 
 	// These currently get linked to actions inside the select tool, but will eventually have their own buttons among the tools
 	// once selection is pulled out to mode-level.
-	UI_COMMAND(SewAction, "Sew", "Sew selected edges to their pairs", EUserInterfaceActionType::Button, FInputChord());
-	UI_COMMAND(SplitAction, "Split", "Split UV's at selected edges", EUserInterfaceActionType::Button, FInputChord());
-	UI_COMMAND(IslandConformalUnwrapAction, "IsdUw", "Perform a  conformal unwrap of selected islands", EUserInterfaceActionType::Button, FInputChord());
+	UI_COMMAND(SewAction, "Sew", "Sew edges. The red edges will be sewn to the green edges", EUserInterfaceActionType::Button, FInputChord());
+	UI_COMMAND(SplitAction, "Split", 
+		"Given an edge selection, split those edges. Given a vertex selection, split any selected bowtie verts.", 
+		EUserInterfaceActionType::Button, FInputChord());
+	UI_COMMAND(IslandConformalUnwrapAction, "IsdUw", "Apply a conformal unwrap to the selected UV islands", EUserInterfaceActionType::Button, FInputChord());
 
 	// These allow us to link up to pressed keys
 	UI_COMMAND(AcceptOrCompleteActiveTool, "Accept", "Accept the active tool", EUserInterfaceActionType::Button, FInputChord(EKeys::Enter));
diff --git a/Engine/Plugins/Experimental/UVEditor/Source/UVEditorTools/Private/UVSeamSewAction.cpp b/Engine/Plugins/Experimental/UVEditor/Source/UVEditorTools/Private/UVSeamSewAction.cpp
--- a/Engine/Plugins/Experimental/UVEditor/Source/UVEditorTools/Private/UVSeamSewAction.cpp
+++ b/Engine/Plugins/Experimental/UVEditor/Source/UVEditorTools/Private/UVSeamSewAction.cpp
@@ -243,7 +243,8 @@ bool UUVSeamSewAction::ApplyAction(UUVToolEmitChangeAPI& EmitChangeAPI)
 {
 	TRACE_CPUPROFILER_EVENT_SCOPE(UVSeamSewAction_ApplyAction);
 
-	FDynamicMesh3& MeshToSew = *(Targets[SelectionTargetIndex]->UnwrapCanonical);
+	UUVEditorToolMeshInput* Target = Targets[SelectionTargetIndex];
+	FDynamicMesh3& MeshToSew = *(Target->UnwrapCanonical);
 
 	TArray<int32> SelectedTids;
 	TArray<FIndex2i> ResolvedEdgePairs;
@@ -269,7 +270,7 @@ bool UUVSeamSewAction::ApplyAction(UUVToolEmitChangeAPI& EmitChangeAPI)
 	SelectedTids.Sort();
 	SelectedTids.SetNum(Algo::Unique(SelectedTids));
 
-	FDynamicMeshChangeTracker ChangeTracker(Targets[SelectionTargetIndex]->UnwrapCanonical.Get());
+	FDynamicMeshChangeTracker ChangeTracker(Target->UnwrapCanonical.Get());
 	ChangeTracker.BeginChange();
 	ChangeTracker.SaveTriangles(SelectedTids, true);
 
@@ -312,16 +313,18 @@ bool UUVSeamSewAction::ApplyAction(UUVToolEmitChangeAPI& EmitChangeAPI)
 		}
 	}
 
-	Targets[SelectionTargetIndex]->UpdateUnwrapCanonicalOverlayFromPositions(&RemainingVids, &SelectedTids);
-	Targets[SelectionTargetIndex]->UpdateAllFromUnwrapCanonical(&RemainingVids, &SelectedTids, &SelectedTids);
-	checkSlow(MeshToSew.IsSameAs(*Targets[SelectionTargetIndex]->UnwrapPreview->PreviewMesh->GetMesh(), FDynamicMesh3::FSameAsOptions()));
+	// Our selection is no longer valid, and we should clear it now before the broadcasts from the upcoming
+	// canonical updates ask us to rebuild our visualization.
+	SetSelection(-1, nullptr);
+
+	Target->UpdateUnwrapCanonicalOverlayFromPositions(&RemainingVids, &SelectedTids);
+	Target->UpdateAllFromUnwrapCanonical(&RemainingVids, &SelectedTids, &SelectedTids);
+	checkSlow(MeshToSew.IsSameAs(*Target->UnwrapPreview->PreviewMesh->GetMesh(), FDynamicMesh3::FSameAsOptions()));
 
 	const FText TransactionName(LOCTEXT("SewCompleteTransactionName", "Sew Edges"));
-	EmitChangeAPI.EmitToolIndependentUnwrapCanonicalChange(Targets[SelectionTargetIndex],
+	EmitChangeAPI.EmitToolIndependentUnwrapCanonicalChange(Target,
 		ChangeTracker.EndChange(), TransactionName);
 
-	SetSelection(-1, nullptr);
-
 	return true;
 }
 
diff --git a/Engine/Plugins/Experimental/UVEditor/Source/UVEditorTools/Private/UVSelectTool.cpp b/Engine/Plugins/Experimental/UVEditor/Source/UVEditorTools/Private/UVSelectTool.cpp
--- a/Engine/Plugins/Experimental/UVEditor/Source/UVEditorTools/Private/UVSelectTool.cpp
+++ b/Engine/Plugins/Experimental/UVEditor/Source/UVEditorTools/Private/UVSelectTool.cpp
@@ -107,25 +107,19 @@ namespace UVSelectToolLocals
 	{
 	public:
 		/**
-		 * @param bBroadcastOnSelectionChangedIn Whether the change in selection should broadcast
-		 *   OnSelectionChanged, which updates gizmo, etc.
-		 * @param GizmoBeforeIn Only relevant if bBroadcastOnSelectionChangedIn is true. In that case,
-		 *   the gizmo gets reset on the way forward to the current selection, which means we have to
-		 *   reset it to the old orientation on the way back (otherwise a rotated gizmo would end up
-		 *   losing its rotation on undo).
+		 * @param GizmoBeforeIn Transform to which to revert gizmo on the way back (to avoid losing gizmo rotation, which
+		 *   gets transacted post-selection-change, and therefore in the wrong place for undo).
 		 * @param EdgeVidPairsBeforeIn
 		 * @param EdgeVidPairsAfterIn
 		 */
 		FSelectionChange(const FDynamicMeshSelection& SelectionBeforeIn,
 			const FDynamicMeshSelection& SelectionAfterIn,
-			bool bBroadcastOnSelectionChangedIn,
 			const FTransform& GizmoBeforeIn,
 			TUniquePtr<TArray<FIndex2i>> EdgeVidPairsBeforeIn,
 			TUniquePtr<TArray<FIndex2i>> EdgeVidPairsAfterIn
 			)
 			: SelectionBefore(SelectionBeforeIn)
 			, SelectionAfter(SelectionAfterIn)
-			, bBroadcastOnSelectionChanged(bBroadcastOnSelectionChangedIn)
 			, GizmoBefore(GizmoBeforeIn)
 			, EdgeVidPairsBefore(MoveTemp(EdgeVidPairsBeforeIn))
 			, EdgeVidPairsAfter(MoveTemp(EdgeVidPairsAfterIn))
@@ -142,7 +136,7 @@ namespace UVSelectToolLocals
 			if (ensure(ChangeRouter) && ChangeRouter->CurrentSelectTool.IsValid())
 			{
 				UpdateSelectionEidsAfterMeshChange(SelectionAfter, EdgeVidPairsAfter.Get());
-				ChangeRouter->CurrentSelectTool->SetSelection(SelectionAfter, bBroadcastOnSelectionChanged);
+				ChangeRouter->CurrentSelectTool->SetSelection(SelectionAfter);
 			}
 		}
 
@@ -152,11 +146,8 @@ namespace UVSelectToolLocals
 			if (ensure(ChangeRouter) && ChangeRouter->CurrentSelectTool.IsValid())
 			{
 				UpdateSelectionEidsAfterMeshChange(SelectionBefore, EdgeVidPairsBefore.Get());
-				ChangeRouter->CurrentSelectTool->SetSelection(SelectionBefore, bBroadcastOnSelectionChanged);
-				if (bBroadcastOnSelectionChanged)
-				{
-					ChangeRouter->CurrentSelectTool->SetGizmoTransform(GizmoBefore);
-				}
+				ChangeRouter->CurrentSelectTool->SetSelection(SelectionBefore);
+				ChangeRouter->CurrentSelectTool->SetGizmoTransform(GizmoBefore);
 			}
 		}
 
@@ -174,7 +165,6 @@ namespace UVSelectToolLocals
 	protected:
 		FDynamicMeshSelection SelectionBefore;
 		FDynamicMeshSelection SelectionAfter;
-		bool bBroadcastOnSelectionChanged;
 		FTransform GizmoBefore;
 
 		TUniquePtr<TArray<FIndex2i>> EdgeVidPairsBefore;
@@ -346,7 +336,7 @@ void UUVSelectTool::Setup()
 	ChangeRouter->CurrentSelectTool = this;
 
 	SelectionMechanic->EmitSelectionChange = [this](const FDynamicMeshSelection& OldSelection,
-		const FDynamicMeshSelection& NewSelection, bool bBroadcastOnSelectionChangedIn)
+		const FDynamicMeshSelection& NewSelection)
 	{
 		TUniquePtr<TArray<FIndex2i>> VidPairsBefore;
 		TUniquePtr<TArray<FIndex2i>> VidPairsAfter;
@@ -361,7 +351,7 @@ void UUVSelectTool::Setup()
 			GetVidPairsFromSelection(NewSelection, *VidPairsAfter);
 		}
 		EmitChangeAPI->EmitToolIndependentChange(ChangeRouter, MakeUnique<UVSelectToolLocals::FSelectionChange>(
-			OldSelection, NewSelection, bBroadcastOnSelectionChangedIn, TransformGizmo->GetGizmoTransform(),
+			OldSelection, NewSelection, TransformGizmo->GetGizmoTransform(),
 			MoveTemp(VidPairsBefore), MoveTemp(VidPairsAfter)),
 			LOCTEXT("SelectionChangeMessage", "Selection Change"));
 	};
@@ -415,6 +405,7 @@ void UUVSelectTool::Setup()
 			UpdateSelectionEidsAfterMeshChange(*SelectionMechanic, &CurrentSelectionVidPairs);
 			UpdateGizmo();
 			SelectionMechanic->RebuildDrawnElements(TransformGizmo->GetGizmoTransform());
+			SewAction->UpdateVisualizations();
 		});
 	}
 
@@ -526,14 +517,13 @@ void UUVSelectTool::Shutdown(EToolShutdownType ShutdownType)
 	ChangeRouter = nullptr;
 }
 
-void UUVSelectTool::SetSelection(const FDynamicMeshSelection& NewSelection, bool bBroadcastOnSelectionChanged)
+void UUVSelectTool::SetSelection(const FDynamicMeshSelection& NewSelection)
 {
-	SelectionMechanic->SetSelection(NewSelection, bBroadcastOnSelectionChanged, 
+	SelectionMechanic->SetSelection(NewSelection, true, 
 		false); // Don't emit undo because this function is called from undo
 	
 	// Make sure the current selection mode is compatible with the new selection we received. Don't broadcast
-	// this part because presumably we've already responded to selection change if bBroadcastOnSelectionChanged
-	// was true above.
+	// this part because presumably we've already responded to selection change through selection change broadcast.
 	// TODO: there are a couple things that are not ideal about the below. One is that we always change to
 	// triangle mode when we don't know if the triangles came from island or mesh selection mode. Another is
 	// that we change the selection mode in the mechanic directly rather than going through ChangeSelectionMode,
@@ -977,32 +967,43 @@ void UUVSelectTool::ApplySplit()
 
 	const FDynamicMeshSelection& Selection = SelectionMechanic->GetCurrentSelection();
 
-	if (Selection.IsEmpty() || Selection.Type != FDynamicMeshSelection::EType::Edge)
+	if (Selection.IsEmpty())
 	{
 		GetToolManager()->DisplayMessage(
-			LOCTEXT("SplitErrorSelectionEmpty", "Cannot split UV's. Edge selection was empty."),
+			LOCTEXT("SplitErrorSelectionEmpty", "Cannot split UVs. Selection was empty."),
 			EToolMessageLevel::UserWarning);
-		return;
 	}
+	else if (Selection.Type == FDynamicMeshSelection::EType::Edge)
+	{
+		ApplySplitEdges();
+	}
+	else if (Selection.Type == FDynamicMeshSelection::EType::Vertex)
+	{
+		ApplySplitBowtieVertices();
+	}
+	else
+	{
+		GetToolManager()->DisplayMessage(
+			LOCTEXT("SplitErrorSelectionEmpty", "Cannot split UVs. Selection must be edges or vertices."),
+			EToolMessageLevel::UserWarning);
+	}
+}
 
-	if (!ensure(SelectionTargetIndex >= 0))
+void UUVSelectTool::ApplySplitEdges()
+{
+	const FDynamicMeshSelection& Selection = SelectionMechanic->GetCurrentSelection();
+	if (!ensure(SelectionTargetIndex >= 0 && !Selection.IsEmpty() && Selection.Type == FDynamicMeshSelection::EType::Edge))
 	{
 		return;
 	}
 	UUVEditorToolMeshInput* Target = Targets[SelectionTargetIndex];
-	
+
 	// Gather up the corresponding edge IDs in the applied (3d) mesh.
 	TSet<int32> AppliedEidSet;
 	for (int32 Eid : Selection.SelectedIDs)
 	{
-		if (Selection.Mesh->IsBoundaryEdge(Eid))
-		{
-			// We will skip these already-split edges here. It would be safe to pass the corresponding
-			// applied edge into the CreateSeamsAtEdges call, but we don't want this edge to stay selected
-			// after the split action, because we would like a split followed by immediate sew to revert
-			// the mesh to the previous state, rather than sewing edges that started out split.
-			continue;
-		}
+		// Note that we don't check whether edges are already boundary edges because we allow such edges
+		// to be selected for splitting of any attached bowties.
 
 		FIndex2i EdgeUnwrapVids = Selection.Mesh->GetEdgeV(Eid);
 
@@ -1016,7 +1017,7 @@ void UUVSelectTool::ApplySplit()
 		}
 	}
 
-	// Perform the cut in the overlay
+	// Perform the cut in the overlay, but don't propagate to unwrap yet
 	FUVEditResult UVEditResult;
 	FDynamicMeshUVEditor UVEditor(Target->AppliedCanonical.Get(),
 		Target->UVLayerIndex, false);
@@ -1035,17 +1036,13 @@ void UUVSelectTool::ApplySplit()
 	ChangeTracker.BeginChange();
 	ChangeTracker.SaveTriangles(TidSet, true);
 
-	// We're about to update the unwrap, which may mess up our selection because a selected edge may
-	// no longer exist after the update, even if we store it as a pair of verts. Instead, we're going
-	// to be changing the selection to all the resulting border edges after this.
-	// The cleanest thing to do (esp for undo/redo) is to clear selection first, then reset it.
-
-	const FText TransactionName(LOCTEXT("ApplySplitTransactionName", "Split Edges"));
+	const FText TransactionName(LOCTEXT("ApplySplitEdgesTransactionName", "Split Edges"));
 	EmitChangeAPI->BeginUndoTransaction(TransactionName);
 
-	FDynamicMeshSelection NewSelection = SelectionMechanic->GetCurrentSelection();
-	FDynamicMeshSelection EmptySelection;
-	SelectionMechanic->SetSelection(EmptySelection, false, true); // don't broadcast, do emit undo
+	// Clear selection here so that it is restored on undo
+	// TODO: We could apply a new selection after the change, but this should only happen
+	// once we apply a new selection on split, which we currently don't do.
+	SelectionMechanic->SetSelection(FDynamicMeshSelection(), true, true);
 
 	// Perform the update
 	TArray<int32> AppliedTids = TidSet.Array();
@@ -1058,26 +1055,68 @@ void UUVSelectTool::ApplySplit()
 	EmitChangeAPI->EmitToolIndependentUnwrapCanonicalChange(
 		Target, ChangeTracker.EndChange(), TransactionName);
 
-	// Set selection to new border edges
-	NewSelection.SelectedIDs.Empty();
-	for (int32 AppliedEid : AppliedEidSet)
+	EmitChangeAPI->EndUndoTransaction();
+}
+
+void UUVSelectTool::ApplySplitBowtieVertices()
+{
+	const FDynamicMeshSelection& Selection = SelectionMechanic->GetCurrentSelection();
+	if (!ensure(SelectionTargetIndex >= 0 && !Selection.IsEmpty() && Selection.Type == FDynamicMeshSelection::EType::Vertex))
 	{
-		FIndex2i EdgeAppliedVids = Target->AppliedCanonical->GetEdgeV(AppliedEid);
-		TArray<int32> UnwrapVids1, UnwrapVids2;
-		Target->AppliedVidToUnwrapVids(EdgeAppliedVids.A, UnwrapVids1);
-		Target->AppliedVidToUnwrapVids(EdgeAppliedVids.B, UnwrapVids2);
-		for (int32 Vid1 : UnwrapVids1)
-		{
-			for (int32 Vid2 : UnwrapVids2)
-			{
-				int32 Eid = NewSelection.Mesh->FindEdge(Vid1, Vid2);
-				if (Eid != IndexConstants::InvalidID)
-				{
-					NewSelection.SelectedIDs.Add(Eid);
-				}
-			}
-		}
+		return;
+	}
+	UUVEditorToolMeshInput* Target = Targets[SelectionTargetIndex];
+
+	// Gather the corresponding vert ID's in the applied mesh
+	TSet<int32> AppliedVidSet;
+	for (int32 UnwrapVid : Selection.SelectedIDs)
+	{
+		int32 AppliedVid = Target->UnwrapVidToAppliedVid(UnwrapVid);
+		AppliedVidSet.Add(AppliedVid);
+	}
+
+	// Split any bowties in the applied mesh overlay
+	TArray<int32> NewUVElements;
+	FDynamicMeshUVOverlay* Overlay = Target->AppliedCanonical->Attributes()->GetUVLayer(Target->UVLayerIndex);
+	for (int32 Vid : AppliedVidSet)
+	{
+		Overlay->SplitBowtiesAtVertex(Vid, &NewUVElements);
+	}
+
+	// Prep for undo transaction
+	TSet<int32> TidSet;
+	for (int32 UnwrapVid : NewUVElements)
+	{
+		TArray<int32> VertTids;
+		Target->AppliedCanonical->GetVtxTriangles(Target->UnwrapVidToAppliedVid(UnwrapVid), VertTids);
+		TidSet.Append(VertTids);
 	}
+
+	FDynamicMeshChangeTracker ChangeTracker(Target->UnwrapCanonical.Get());
+	ChangeTracker.BeginChange();
+	ChangeTracker.SaveTriangles(TidSet, true);
+
+	const FText TransactionName(LOCTEXT("ApplySplitBowtieVerticesTransactionName", "Split Bowties"));
+	EmitChangeAPI->BeginUndoTransaction(TransactionName);
+
+	// Emit selection clear first so that we restore it on undo
+	FDynamicMeshSelection NewSelection = SelectionMechanic->GetCurrentSelection(); // save type, etc
+	// TODO: This emitted transaction doesn't actually need to broadcast on redo, but we don't yet
+	// have support for that.
+	SelectionMechanic->SetSelection(FDynamicMeshSelection(), false, true); // don't broadcast, do emit undo
+
+	// Perform the update
+	TArray<int32> AppliedTids = TidSet.Array();
+	Target->UpdateAllFromAppliedCanonical(&NewUVElements, &AppliedTids, &AppliedTids);
+
+	// Emit update transaction
+	EmitChangeAPI->EmitToolIndependentUnwrapCanonicalChange(
+		Target, ChangeTracker.EndChange(), TransactionName);
+
+	// Set up the new selection to include the new elements
+	NewSelection.SelectedIDs.Append(NewUVElements);
+	// TODO: This emitted transaction doesn't actually need to broadcast on undo, but we don't yet
+	// have support for that.
 	SelectionMechanic->SetSelection(NewSelection, true, true); // both broadcast and emit undo
 
 	EmitChangeAPI->EndUndoTransaction();
diff --git a/Engine/Plugins/Experimental/UVEditor/Source/UVEditorTools/Public/UVSelectTool.h b/Engine/Plugins/Experimental/UVEditor/Source/UVEditorTools/Public/UVSelectTool.h
--- a/Engine/Plugins/Experimental/UVEditor/Source/UVEditorTools/Public/UVSelectTool.h
+++ b/Engine/Plugins/Experimental/UVEditor/Source/UVEditorTools/Public/UVSelectTool.h
@@ -68,11 +68,15 @@ public:
 	void Initialize(UUVSelectTool* ParentToolIn) { ParentTool = ParentToolIn; }
 	void PostAction(ESelectToolAction Action);
 
+	//~ NOTE: the below comments don't actually get used as tooltips because we use a detail
+	//~ customization, and the tooltips are specified in UVEditorCommands.cpp. Change them
+	//~ there if needed!
+
 	/** Sew edges. The red edges will be sewn to the green edges */
 	UFUNCTION(CallInEditor, Category = Actions)
 	void Sew() { PostAction(ESelectToolAction::Sew); }
 
-	/** Split given edges/verts */
+	/** Given an edge selection, split those edges. Given a vertex selection, split any selected bowtie verts. */
 	UFUNCTION(CallInEditor, Category = Actions)
 	void Split() { PostAction(ESelectToolAction::Split); };
 
@@ -108,7 +112,7 @@ public:
 	}
 
 	// Used by undo/redo changes to update the tool state
-	void SetSelection(const UE::Geometry::FDynamicMeshSelection& NewSelection, bool bBroadcastOnSelectionChanged);
+	void SetSelection(const UE::Geometry::FDynamicMeshSelection& NewSelection);
 	void SetGizmoTransform(const FTransform& NewTransform);
 
 	// UInteractiveTool
@@ -197,6 +201,8 @@ protected:
 
 	void ApplyAction(ESelectToolAction ActionType);
 	void ApplySplit();
+	void ApplySplitEdges();
+	void ApplySplitBowtieVertices();
 
 
 };
diff --git a/Engine/Plugins/Runtime/GeometryProcessing/Source/DynamicMesh/Private/Parameterization/DynamicMeshUVEditor.cpp b/Engine/Plugins/Runtime/GeometryProcessing/Source/DynamicMesh/Private/Parameterization/DynamicMeshUVEditor.cpp
--- a/Engine/Plugins/Runtime/GeometryProcessing/Source/DynamicMesh/Private/Parameterization/DynamicMeshUVEditor.cpp
+++ b/Engine/Plugins/Runtime/GeometryProcessing/Source/DynamicMesh/Private/Parameterization/DynamicMeshUVEditor.cpp
@@ -846,46 +846,68 @@ bool FDynamicMeshUVEditor::CreateSeamsAtEdges(const TSet<int32>& EidsToMakeIntoS
 
 		if (UVOverlay->IsSeamEdge(Eid))
 		{
-			// If this is already a seam edge, we may have made it one by splitting the opposite vertex earlier
-			// (as part of an adjacent edge split). Make sure that we don't leave a bowtie.
+			// If this is already a seam edge, make sure its endpoints are not bowties. We create such edges 
+			// sometimes as we split adjacent edges, and it is also likely to be what the user wants when
+			// selecting edges adjacent to a bowtie.
 			for (int i = 0; i < 2; ++i)
 			{
 				if (UVOverlay->IsBowtieInOverlay(EdgeVids[i]))
 				{
-					VertNeedsSplitting[i] = true;
-
-					// Gather up the attached verts into VertTidsToUseForSplit
-					ensure(DoesVertHaveAnotherSeamAttached(Eid, EdgeVids[i], VertTidsToUseForSplit[i]));
+					UVOverlay->SplitBowtiesAtVertex(EdgeVids[i], 
+						Result ? &Result->NewUVElements : nullptr);
 				}
 			}
+			continue;
 		}
-		else
+
+		// If we're not already a seam, then it's going to become one. A vert needs to get split if it
+		// has a present or future seam attached, or else it will become a bowtie if we just split the
+		// other vert.
+		for (int i = 0; i < 2; ++i)
 		{
-			// If we're not already a seam, then it's going to become one. A vert needs to get split if it
-			// has a present or future seam attached, or else it will become a bowtie if we just split the
-			// other vert.
-			for (int i = 0; i < 2; ++i)
-			{
-				VertNeedsSplitting[i] = DoesVertHaveAnotherSeamAttached(Eid, EdgeVids[i], VertTidsToUseForSplit[i]);
-			}
+			VertNeedsSplitting[i] = DoesVertHaveAnotherSeamAttached(Eid, EdgeVids[i], VertTidsToUseForSplit[i]);
+		}
 			
-			// If neither absolutely has to get split, then one needs to get split anyway so that we
-			// make the edge into a seam.
-			if (!VertNeedsSplitting[0] && !VertNeedsSplitting[1])
+		// If neither absolutely has to get split, then one needs to get split anyway so that we
+		// make the edge into a seam.
+		if (!VertNeedsSplitting[0] && !VertNeedsSplitting[1])
+		{
+			// We'll go halfway around the triangle one-ring, rounding up
+			int32 NumTrisToDisconnect = (VertTidsToUseForSplit[0].Num() + 1) / 2;
+
+			// In doing this, we are splitting an adjacent edge and therefore might inadvertantly
+			// introduce a bowtie on the other vert of that edge, which we would like to avoid. 
+			// So keep track of that potential bowtie.
+			int32 OtherSplitEdge = Mesh->FindEdgeFromTriPair(
+				VertTidsToUseForSplit[0][NumTrisToDisconnect - 1],
+				VertTidsToUseForSplit[0][NumTrisToDisconnect]);
+			int32 PotentialBowtieVid = ensure(OtherSplitEdge != IndexConstants::InvalidID) ?
+				IndexUtil::FindEdgeOtherVertex(Mesh->GetEdgeV(OtherSplitEdge), EdgeVids[0])
+				: IndexConstants::InvalidID;
+
+			// Perform the split
+			VertTidsToUseForSplit[0].SetNum(NumTrisToDisconnect);
+			SplitEdgeVertElement(Eid, EdgeVids[0], VertTidsToUseForSplit[0]);
+
+			// Deal with the bowtie if we created one
+			if (ensure(PotentialBowtieVid != IndexConstants::InvalidID) 
+				&& UVOverlay->IsBowtieInOverlay(PotentialBowtieVid))
 			{
-				VertNeedsSplitting[0] = true;
-				VertTidsToUseForSplit[0].SetNum((VertTidsToUseForSplit[0].Num() + 1) / 2); // Go halfway around, rounding up
+				UVOverlay->SplitBowtiesAtVertex(PotentialBowtieVid,
+					Result ? &Result->NewUVElements : nullptr);
 			}
 		}
-
-		for (int i = 0; i < 2; ++i)
+		else
 		{
-			if (VertNeedsSplitting[i])
+			for (int i = 0; i < 2; ++i)
 			{
-				SplitEdgeVertElement(Eid, EdgeVids[i], VertTidsToUseForSplit[i]);
+				if (VertNeedsSplitting[i])
+				{
+					SplitEdgeVertElement(Eid, EdgeVids[i], VertTidsToUseForSplit[i]);
+				}
 			}
-		}
-	}
+		}//end if at least one side needed splitting
+	}//end for each edge
 
 	return true;
 }
diff --git a/Engine/Plugins/Runtime/MeshModelingToolset/Source/ModelingComponents/Private/Selection/MeshSelectionMechanic.cpp b/Engine/Plugins/Runtime/MeshModelingToolset/Source/ModelingComponents/Private/Selection/MeshSelectionMechanic.cpp
--- a/Engine/Plugins/Runtime/MeshModelingToolset/Source/ModelingComponents/Private/Selection/MeshSelectionMechanic.cpp
+++ b/Engine/Plugins/Runtime/MeshModelingToolset/Source/ModelingComponents/Private/Selection/MeshSelectionMechanic.cpp
@@ -31,10 +31,9 @@ namespace MeshSelectionMechanicLocals
 	{
 	public:
 		FMeshSelectionMechanicSelectionChange(const FDynamicMeshSelection& OldSelection,
-			const FDynamicMeshSelection& NewSelection, bool bBroadcastOnSelectionChangedIn)
+			const FDynamicMeshSelection& NewSelection)
 			: Before(OldSelection)
 			, After(NewSelection)
-			, bBroadcastOnSelectionChanged(bBroadcastOnSelectionChangedIn)
 		{
 		}
 
@@ -43,7 +42,7 @@ namespace MeshSelectionMechanicLocals
 			UMeshSelectionMechanic* Mechanic = Cast<UMeshSelectionMechanic>(Object);
 			if (Mechanic)
 			{
-				Mechanic->SetSelection(After, bBroadcastOnSelectionChanged, false);
+				Mechanic->SetSelection(After, true, false);
 			}
 		}
 
@@ -52,7 +51,7 @@ namespace MeshSelectionMechanicLocals
 			UMeshSelectionMechanic* Mechanic = Cast<UMeshSelectionMechanic>(Object);
 			if (Mechanic)
 			{
-				Mechanic->SetSelection(Before, bBroadcastOnSelectionChanged, false);
+				Mechanic->SetSelection(Before, true, false);
 			}
 		}
 
@@ -64,7 +63,6 @@ namespace MeshSelectionMechanicLocals
 	protected:
 		FDynamicMeshSelection Before;
 		FDynamicMeshSelection After;
-		bool bBroadcastOnSelectionChanged;
 	};
 
 	template <typename InElementType>
@@ -313,10 +311,10 @@ void UMeshSelectionMechanic::Setup(UInteractiveTool* ParentToolIn)
 	// Add default selection change emitter if one was not provided.
 	if (!EmitSelectionChange)
 	{
-		EmitSelectionChange = [this](const FDynamicMeshSelection& OldSelection, const FDynamicMeshSelection& NewSelection, bool bBroadcastOnSelectionChanged)
+		EmitSelectionChange = [this](const FDynamicMeshSelection& OldSelection, const FDynamicMeshSelection& NewSelection)
 		{
 			GetParentTool()->GetToolManager()->EmitObjectChange(this, 
-				MakeUnique<MeshSelectionMechanicLocals::FMeshSelectionMechanicSelectionChange>(OldSelection, NewSelection, bBroadcastOnSelectionChanged),
+				MakeUnique<MeshSelectionMechanicLocals::FMeshSelectionMechanicSelectionChange>(OldSelection, NewSelection),
 				LOCTEXT("SelectionChangeMessage", "Selection Change"));
 		};
 	}
@@ -442,7 +440,7 @@ void UMeshSelectionMechanic::SetSelection(const FDynamicMeshSelection& Selection
 
 	if (bEmitChange && OriginalSelection != CurrentSelection)
 	{
-		EmitSelectionChange(OriginalSelection, CurrentSelection, bBroadcast);
+		EmitSelectionChange(OriginalSelection, CurrentSelection);
 	}
 	if (bBroadcast)
 	{
@@ -829,7 +827,7 @@ void UMeshSelectionMechanic::ChangeSelectionMode(const EMeshSelectionMechanicMod
 	{
 		UpdateCentroid();
 		RebuildDrawnElements(FTransform(GetCurrentSelectionCentroid()));
-		EmitSelectionChange(OriginalSelection, CurrentSelection, true);
+		EmitSelectionChange(OriginalSelection, CurrentSelection);
 		OnSelectionChanged.Broadcast();
 	}
 
@@ -1001,7 +999,7 @@ void UMeshSelectionMechanic::OnClicked(const FInputDeviceRay& ClickPos)
 	if (OriginalSelection != CurrentSelection)
 	{
 		UpdateCentroid();
-		EmitSelectionChange(OriginalSelection, CurrentSelection, true);
+		EmitSelectionChange(OriginalSelection, CurrentSelection);
 		OnSelectionChanged.Broadcast();
 
 		// Rebuild after broadcast in case the outside world wants to adjust things like color...
@@ -1169,7 +1167,7 @@ void UMeshSelectionMechanic::OnDragRectangleFinished(const FCameraRectangle& Cur
 	if (!bCancelled && (PreDragSelection != CurrentSelection))
 	{
 		UpdateCentroid();		
-		EmitSelectionChange(PreDragSelection, CurrentSelection, true);
+		EmitSelectionChange(PreDragSelection, CurrentSelection);
 		OnSelectionChanged.Broadcast();
 
 		// Rebuild after broadcast in case the outside world wants to adjust things like color...
diff --git a/Engine/Plugins/Runtime/MeshModelingToolset/Source/ModelingComponents/Public/Selection/MeshSelectionMechanic.h b/Engine/Plugins/Runtime/MeshModelingToolset/Source/ModelingComponents/Public/Selection/MeshSelectionMechanic.h
--- a/Engine/Plugins/Runtime/MeshModelingToolset/Source/ModelingComponents/Public/Selection/MeshSelectionMechanic.h
+++ b/Engine/Plugins/Runtime/MeshModelingToolset/Source/ModelingComponents/Public/Selection/MeshSelectionMechanic.h
@@ -116,12 +116,11 @@ public:
 	/**
 	 * Function to use for emitting selection change events. If not set, Setup() will attach a version that
 	 * uses EmitObjectChange() on the tool manager to emit a change that operates on this mechanic.
-	 * @param bBroadcastOnSelectionChanged Indicates whether the SetSelection call that asked to issue the
-	 *   transaction was also requested to broadcast the OnSelectionChanged delegate. If so, the transaction
-	 *   should probably do the same on revert/apply.
+	 * TODO: User should probably be able to specify whether the emitted transaction should broadcast
+	 * OnSelectionChanged on redo, undo, or both, to allow selection change events to be used as bookends
+	 * around topology changes.
 	 */ 
-	TUniqueFunction<void(const FDynamicMeshSelection& OldSelection, const FDynamicMeshSelection& NewSelection, 
-		bool bBroadcastOnSelectionChanged)> EmitSelectionChange;
+	TUniqueFunction<void(const FDynamicMeshSelection& OldSelection, const FDynamicMeshSelection& NewSelection)> EmitSelectionChange;
 
 protected:
 
diff --git a/Engine/Source/Runtime/GeometryCore/Private/DynamicMesh/DynamicMeshOverlay.cpp b/Engine/Source/Runtime/GeometryCore/Private/DynamicMesh/DynamicMeshOverlay.cpp
--- a/Engine/Source/Runtime/GeometryCore/Private/DynamicMesh/DynamicMeshOverlay.cpp
+++ b/Engine/Source/Runtime/GeometryCore/Private/DynamicMesh/DynamicMeshOverlay.cpp
@@ -250,105 +250,118 @@ template<typename RealType, int ElementSize>
 void TDynamicMeshOverlay<RealType, ElementSize>::SplitBowties()
 {
 	TRACE_CPUPROFILER_EVENT_SCOPE(DynamicMeshOverlay_SplitBowties);
-	
+
+	for (int VertexID : ParentMesh->VertexIndicesItr())
+	{
+		SplitBowtiesAtVertex(VertexID);
+	}
+}
+
+
+template<typename RealType, int ElementSize>
+void TDynamicMeshOverlay<RealType, ElementSize>::SplitBowtiesAtVertex(int32 VertexID, TArray<int32>* NewElementIDs)
+{
 	// arrays for storing contiguous triangle groups from parentmesh
 	TArray<int> TrianglesOut, ContiguousGroupLengths;
 	TArray<bool> GroupIsLoop;
 
 	// per-vertex element group tracking data, reused in loop below
 	TSet<int> ElementIDSeen;
-	TArray<int> GroupElementIDs, // stores element IDs 1:1 w/ contiguous triangles in the parent mesh
+	TArray<int> GroupElementIDs, // stores element IDs of this vertex, 1:1 w/ contiguous triangles in the parent mesh
 		SubGroupID, // mapping from GroupElementIDs indices into SubGroupElementIDs indices, giving subgroup membership per triangle
 		SubGroupElementIDs; // 1:1 w/ 'subgroups' in the group (e.g. if all triangles in the group all had the same element ID, this would be an array of length 1, just containing that element ID)
-	
-	for (int VertexID : ParentMesh->VertexIndicesItr())
+
+	ensure(EMeshResult::Ok == ParentMesh->GetVtxContiguousTriangles(VertexID, TrianglesOut, ContiguousGroupLengths, GroupIsLoop));
+	int32 NumTris = TrianglesOut.Num();
+
+	ElementIDSeen.Reset();
+	// per contiguous group of triangles around vertex in ParentMesh, find contiguous sub-groups in overlay
+	for (int32 GroupIdx = 0, NumGroups = ContiguousGroupLengths.Num(), TriSubStart = 0; GroupIdx < NumGroups; GroupIdx++)
 	{
-		ensure(EMeshResult::Ok == ParentMesh->GetVtxContiguousTriangles(VertexID, TrianglesOut, ContiguousGroupLengths, GroupIsLoop));
-		int32 NumTris = TrianglesOut.Num();
+		bool bIsLoop = GroupIsLoop[GroupIdx];
+		int TriInGroupNum = ContiguousGroupLengths[GroupIdx];
+		if (ensure(TriInGroupNum > 0) == false)
+		{
+			continue;
+		}
+		int TriSubEnd = TriSubStart + TriInGroupNum;
 
-		ElementIDSeen.Reset();
-		// per contiguous group of triangles around vertex in ParentMesh, find contiguous sub-groups in overlay
-		for (int32 GroupIdx = 0, NumGroups = ContiguousGroupLengths.Num(), TriSubStart = 0; GroupIdx < NumGroups; GroupIdx++)
+		GroupElementIDs.Reset();
+		for (int TriSubIdx = TriSubStart; TriSubIdx < TriSubEnd; TriSubIdx++)
 		{
-			bool bIsLoop = GroupIsLoop[GroupIdx];
-			int TriInGroupNum = ContiguousGroupLengths[GroupIdx];
-			if (ensure(TriInGroupNum > 0) == false)
-			{
-				continue;
-			}
-			int TriSubEnd = TriSubStart + TriInGroupNum;
-			
-			GroupElementIDs.Reset();
-			for (int TriSubIdx = TriSubStart; TriSubIdx < TriSubEnd; TriSubIdx++)
+			int TriID = TrianglesOut[TriSubIdx];
+			FIndex3i TriVIDs = ParentMesh->GetTriangle(TriID);
+			FIndex3i TriEIDs = GetTriangle(TriID);
+			int SubIdx = TriVIDs.IndexOf(VertexID);
+			GroupElementIDs.Add(TriEIDs[SubIdx]);
+		}
+
+		auto IsConnected = [this, &GroupElementIDs, &TrianglesOut, &TriSubStart](int TriOutIdxA, int TriOutIdxB)
+		{
+			if (GroupElementIDs[TriOutIdxA - TriSubStart] != GroupElementIDs[TriOutIdxB - TriSubStart])
 			{
-				int TriID = TrianglesOut[TriSubIdx];
-				FIndex3i TriVIDs = ParentMesh->GetTriangle(TriID);
-				FIndex3i TriEIDs = GetTriangle(TriID);
-				int SubIdx = TriVIDs.IndexOf(VertexID);
-				GroupElementIDs.Add(TriEIDs[SubIdx]);
+				return false;
 			}
+			int EdgeID = ParentMesh->FindEdgeFromTriPair(TrianglesOut[TriOutIdxA], TrianglesOut[TriOutIdxB]);
+			return EdgeID >= 0 && !IsSeamEdge(EdgeID);
+		};
 
-			auto IsConnected = [this, &GroupElementIDs, &TrianglesOut, &TriSubStart](int TriOutIdxA, int TriOutIdxB)
-			{
-				if (GroupElementIDs[TriOutIdxA - TriSubStart] != GroupElementIDs[TriOutIdxB - TriSubStart])
-				{
-					return false;
-				}
-				int EdgeID = ParentMesh->FindEdgeFromTriPair(TrianglesOut[TriOutIdxA], TrianglesOut[TriOutIdxB]);
-				return EdgeID >= 0 && !IsSeamEdge(EdgeID);
-			};
+		SubGroupID.Reset(); SubGroupID.SetNum(TriInGroupNum);
+		SubGroupElementIDs.Reset();
+		int MaxSubID = 0;
+		SubGroupID[0] = 0;
+		SubGroupElementIDs.Add(GroupElementIDs[0]);
 
-			SubGroupID.Reset(); SubGroupID.SetNum(TriInGroupNum);
-			SubGroupElementIDs.Reset();
-			int MaxSubID = 0;
-			SubGroupID[0] = 0;
-			SubGroupElementIDs.Add(GroupElementIDs[0]);
-			for (int TriSubIdx = TriSubStart; TriSubIdx+1 < TriSubEnd; TriSubIdx++)
+		// Iterate through tris in current group, except last one
+		for (int TriSubIdx = TriSubStart; TriSubIdx + 1 < TriSubEnd; TriSubIdx++)
+		{
+			if (!IsConnected(TriSubIdx, TriSubIdx + 1))
 			{
-				if (!IsConnected(TriSubIdx, TriSubIdx+1))
-				{
-					SubGroupElementIDs.Add(GroupElementIDs[TriSubIdx + 1 - TriSubStart]);
-					MaxSubID++;
-				}
-				SubGroupID[TriSubIdx - TriSubStart + 1] = MaxSubID;
+				SubGroupElementIDs.Add(GroupElementIDs[TriSubIdx + 1 - TriSubStart]);
+				MaxSubID++;
 			}
-			// if group was a loop, need to check if the last sub-group and first sub-group were actually the same group
-			if (bIsLoop && MaxSubID > 0 && IsConnected(TriSubStart, TriSubStart + TriInGroupNum - 1))
+			SubGroupID[TriSubIdx - TriSubStart + 1] = MaxSubID;
+		}
+		// if group was a loop, need to check if the last sub-group and first sub-group were actually the same group
+		if (bIsLoop && MaxSubID > 0 && IsConnected(TriSubStart, TriSubStart + TriInGroupNum - 1))
+		{
+			int LastGroupID = SubGroupID.Last();
+			for (int32 Idx = SubGroupID.Num() - 1; Idx >= 0 && SubGroupID[Idx] == LastGroupID; Idx--)
 			{
-				int LastGroupID = SubGroupID.Last();
-				for (int32 Idx = SubGroupID.Num() - 1; Idx >= 0 && SubGroupID[Idx] == LastGroupID; Idx--)
-				{
-					SubGroupID[Idx] = 0;
-				}
-				MaxSubID--;
-				SubGroupElementIDs.Pop(false);
+				SubGroupID[Idx] = 0;
 			}
+			MaxSubID--;
+			SubGroupElementIDs.Pop(false);
+		}
 
-			for (int SubID = 0; SubID < SubGroupElementIDs.Num(); SubID++)
+		for (int SubID = 0; SubID < SubGroupElementIDs.Num(); SubID++)
+		{
+			int ElementID = SubGroupElementIDs[SubID];
+			if (ElementID < 0)
 			{
-				int ElementID = SubGroupElementIDs[SubID];
-				if (ElementID < 0)
-				{
-					continue;		// skip if this is an invalid ElementID (eg from an invalid triangle)
-				}
-				// split needed the *second* time we see a sub-group using a given ElementID
-				if (ElementIDSeen.Contains(ElementID))
+				continue;		// skip if this is an invalid ElementID (eg from an invalid triangle)
+			}
+			// split needed the *second* time we see a sub-group using a given ElementID
+			if (ElementIDSeen.Contains(ElementID))
+			{
+				TArray<int> ConnectedTris;
+				for (int TriSubIdx = TriSubStart; TriSubIdx < TriSubEnd; TriSubIdx++)
 				{
-					TArray<int> ConnectedTris;
-					for (int TriSubIdx = TriSubStart; TriSubIdx < TriSubEnd; TriSubIdx++)
+					if (SubID == SubGroupID[TriSubIdx - TriSubStart])
 					{
-						if (SubID == SubGroupID[TriSubIdx - TriSubStart])
-						{
-							ConnectedTris.Add(TrianglesOut[TriSubIdx]);
-						}
+						ConnectedTris.Add(TrianglesOut[TriSubIdx]);
 					}
-					SplitElement(ElementID, ConnectedTris);
 				}
-				ElementIDSeen.Add(ElementID);
+				int32 NewElementID = SplitElement(ElementID, ConnectedTris);
+				if (NewElementIDs)
+				{
+					NewElementIDs->Add(NewElementID);
+				}
 			}
-
-			TriSubStart = TriSubEnd;
+			ElementIDSeen.Add(ElementID);
 		}
+
+		TriSubStart = TriSubEnd;
 	}
 }
 
diff --git a/Engine/Source/Runtime/GeometryCore/Public/DynamicMesh/DynamicMeshOverlay.h b/Engine/Source/Runtime/GeometryCore/Public/DynamicMesh/DynamicMeshOverlay.h
--- a/Engine/Source/Runtime/GeometryCore/Public/DynamicMesh/DynamicMeshOverlay.h
+++ b/Engine/Source/Runtime/GeometryCore/Public/DynamicMesh/DynamicMeshOverlay.h
@@ -348,6 +348,13 @@ public:
 	*/
 	int SplitElementWithNewParent(int ElementID, int SplitParentVertexID, const TArrayView<const int>& TrianglesToUpdate);
 
+	/**
+	 * Split any bowties at given vertex.
+	 * 
+	 * @param NewElementIDs If not null, newly created element IDs are placed here. Note that this array is
+	 *   intentionally not cleared before appending to it.
+	 */
+	void SplitBowtiesAtVertex(int32 Vid, TArray<int32>* NewElementIDs = nullptr);
 
 	/**
 	* Refine an existing overlay topology by splitting any bow ties
