UVEditor: Add hotkey to toggle display or hiding of background in the UVEditor viewport.

#rnx

diff --git a/Engine/Plugins/Experimental/UVEditor/Source/UVEditor/Private/UVEditorCommands.cpp b/Engine/Plugins/Experimental/UVEditor/Source/UVEditor/Private/UVEditorCommands.cpp
--- a/Engine/Plugins/Experimental/UVEditor/Source/UVEditor/Private/UVEditorCommands.cpp
+++ b/Engine/Plugins/Experimental/UVEditor/Source/UVEditor/Private/UVEditorCommands.cpp
@@ -52,6 +52,8 @@ void FUVEditorCommands::RegisterCommands()
 
 	UI_COMMAND(EnableOrbitCamera, "Orbit", "Enable Orbit Camera", EUserInterfaceActionType::ToggleButton, FInputChord());
 	UI_COMMAND(EnableFlyCamera, "Fly", "Enable Fly Camera", EUserInterfaceActionType::ToggleButton, FInputChord());
+	UI_COMMAND(ToggleBackground, "Toggle Background", "Toggle the display of the background", EUserInterfaceActionType::ToggleButton, FInputChord(EModifierKey::Alt, EKeys::B));
 }
 
+
 #undef LOCTEXT_NAMESPACE
diff --git a/Engine/Plugins/Experimental/UVEditor/Source/UVEditor/Private/UVEditorMode.cpp b/Engine/Plugins/Experimental/UVEditor/Source/UVEditor/Private/UVEditorMode.cpp
--- a/Engine/Plugins/Experimental/UVEditor/Source/UVEditor/Private/UVEditorMode.cpp
+++ b/Engine/Plugins/Experimental/UVEditor/Source/UVEditor/Private/UVEditorMode.cpp
@@ -239,6 +239,11 @@ void UUVEditorMode::BindCommands()
 	const FUVEditorCommands& CommandInfos = FUVEditorCommands::Get();
 	const TSharedRef<FUICommandList>& CommandList = Toolkit->GetToolkitCommands();
 
+	// Hookup Background toggle command
+	CommandList->MapAction(CommandInfos.ToggleBackground, FExecuteAction::CreateLambda([this]() {
+		BackgroundVisualization->Settings->bVisible = !BackgroundVisualization->Settings->bVisible;
+	}));
+
 	// Hook up to Enter/Esc key presses
 	CommandList->MapAction(
 		CommandInfos.AcceptOrCompleteActiveTool,
diff --git a/Engine/Plugins/Experimental/UVEditor/Source/UVEditor/Public/UVEditorBackgroundPreview.h b/Engine/Plugins/Experimental/UVEditor/Source/UVEditor/Public/UVEditorBackgroundPreview.h
--- a/Engine/Plugins/Experimental/UVEditor/Source/UVEditor/Public/UVEditorBackgroundPreview.h
+++ b/Engine/Plugins/Experimental/UVEditor/Source/UVEditor/Public/UVEditorBackgroundPreview.h
@@ -29,8 +29,8 @@ class UVEDITOR_API UUVEditorBackgroundPreviewProperties : public UInteractiveToo
 {
 	GENERATED_BODY()
 public:
-	/** Should the background be shown */
-	UPROPERTY(EditAnywhere, Category = Background, meta = (DisplayName = "Display Background") )
+	/** Should the background be shown (Alt+B)*/
+	UPROPERTY(EditAnywhere, Category = Background, meta = (DisplayName = "Display Background"))
 	bool bVisible = false;
 
 	/** Source of background visuals */
diff --git a/Engine/Plugins/Experimental/UVEditor/Source/UVEditor/Public/UVEditorCommands.h b/Engine/Plugins/Experimental/UVEditor/Source/UVEditor/Public/UVEditorCommands.h
--- a/Engine/Plugins/Experimental/UVEditor/Source/UVEditor/Public/UVEditorCommands.h
+++ b/Engine/Plugins/Experimental/UVEditor/Source/UVEditor/Public/UVEditorCommands.h
@@ -40,4 +40,5 @@ public:
 
 	TSharedPtr<FUICommandInfo> EnableOrbitCamera;
 	TSharedPtr<FUICommandInfo> EnableFlyCamera;
+	TSharedPtr<FUICommandInfo> ToggleBackground;
 };
\ No newline at end of file
