UVEditor: Make it possible to use selection hotkeys while not focused on the 2d viewport.

#rnx

diff --git a/Engine/Plugins/Experimental/UVEditor/Source/UVEditor/Private/UVEditorToolkit.cpp b/Engine/Plugins/Experimental/UVEditor/Source/UVEditor/Private/UVEditorToolkit.cpp
--- a/Engine/Plugins/Experimental/UVEditor/Source/UVEditor/Private/UVEditorToolkit.cpp
+++ b/Engine/Plugins/Experimental/UVEditor/Source/UVEditor/Private/UVEditorToolkit.cpp
@@ -500,6 +500,16 @@ void FUVEditorToolkit::PostInitAssetEditor()
 	// TODO: This should not be hardcoded
 	LivePreviewViewportClient->SetViewLocation(FVector(-200, 100, 100));
 	LivePreviewViewportClient->SetLookAtLocation(FVector(0, 0, 0));
+
+	// Hook up the viewport command list to our toolkit command list so that hotkeys not
+	// handled by our toolkit would be handled by the viewport (to allow us to use
+	// whatever hotkeys the viewport registered after clicking in the detail panel or
+	// elsewhere in the UV editor).
+	// Note that the "Append" call for a command list should probably have been called
+	// "AppendTo" because it adds the callee object as a child of the argument command list.
+	// I.e. after looking in ToolkitCommands, we will look in the viewport command list.
+	TSharedPtr<SUVEditor2DViewport> ViewportWidget = StaticCastSharedPtr<SUVEditor2DViewport>(ViewportTabContent->GetFirstViewport());
+	ViewportWidget->GetCommandList()->Append(ToolkitCommands);
 }
 
 void FUVEditorToolkit::OnToolkitHostingStarted(const TSharedRef<IToolkit>& Toolkit)
