UVEditor: Change UV editor tab icon.

#rnx

diff --git a/Engine/Plugins/Experimental/UVEditor/Source/UVEditor/Private/UVEditorToolkit.cpp b/Engine/Plugins/Experimental/UVEditor/Source/UVEditor/Private/UVEditorToolkit.cpp
--- a/Engine/Plugins/Experimental/UVEditor/Source/UVEditor/Private/UVEditorToolkit.cpp
+++ b/Engine/Plugins/Experimental/UVEditor/Source/UVEditor/Private/UVEditorToolkit.cpp
@@ -512,6 +512,11 @@ void FUVEditorToolkit::PostInitAssetEditor()
 	ViewportWidget->GetCommandList()->Append(ToolkitCommands);
 }
 
+const FSlateBrush* FUVEditorToolkit::GetDefaultTabIcon() const
+{
+	return FUVEditorStyle::Get().GetBrush("UVEditor.OpenUVEditor");
+}
+
 void FUVEditorToolkit::OnToolkitHostingStarted(const TSharedRef<IToolkit>& Toolkit)
 {
 	ModeUILayer->OnToolkitHostingStarted(Toolkit);
diff --git a/Engine/Plugins/Experimental/UVEditor/Source/UVEditor/Public/UVEditorToolkit.h b/Engine/Plugins/Experimental/UVEditor/Source/UVEditor/Public/UVEditorToolkit.h
--- a/Engine/Plugins/Experimental/UVEditor/Source/UVEditor/Public/UVEditorToolkit.h
+++ b/Engine/Plugins/Experimental/UVEditor/Source/UVEditor/Public/UVEditorToolkit.h
@@ -69,6 +69,7 @@ protected:
 
 	// FAssetEditorToolkit
 	virtual void PostInitAssetEditor() override;
+	virtual const FSlateBrush* GetDefaultTabIcon() const override;
 
 	/** Scene in which the 2D unwrapped uv meshes live. */
 	TUniquePtr<FPreviewScene> UnwrapScene;
