UVEditor: Tweak button and property labels

diff --git a/Engine/Plugins/Experimental/UVEditor/Source/UVEditor/Private/DetailsCustomizations/UVSelectToolCustomizations.cpp b/Engine/Plugins/Experimental/UVEditor/Source/UVEditor/Private/DetailsCustomizations/UVSelectToolCustomizations.cpp
--- a/Engine/Plugins/Experimental/UVEditor/Source/UVEditor/Private/DetailsCustomizations/UVSelectToolCustomizations.cpp
+++ b/Engine/Plugins/Experimental/UVEditor/Source/UVEditor/Private/DetailsCustomizations/UVSelectToolCustomizations.cpp
@@ -71,8 +71,8 @@ void FUVSelectToolActionPropertySetDetails::CustomizeDetails(IDetailLayoutBuilde
 	ToolbarBuilder.AddToolBarButton(CommandInfos.SplitAction);
 	ToolbarBuilder.AddToolBarButton(CommandInfos.IslandConformalUnwrapAction);
 
-	DetailBuilder.EditCategory("SelectionActions")
-		.AddCustomRow(LOCTEXT("ActionsSectionFilterString", "Selection Actions"), false)
+	DetailBuilder.EditCategory("EditActions")
+		.AddCustomRow(LOCTEXT("ActionsSectionFilterString", "Edit Actions"), false)
 		[
 			ToolbarBuilder.MakeWidget()
 		];
diff --git a/Engine/Plugins/Experimental/UVEditor/Source/UVEditor/Private/UVEditorCommands.cpp b/Engine/Plugins/Experimental/UVEditor/Source/UVEditor/Private/UVEditorCommands.cpp
--- a/Engine/Plugins/Experimental/UVEditor/Source/UVEditor/Private/UVEditorCommands.cpp
+++ b/Engine/Plugins/Experimental/UVEditor/Source/UVEditor/Private/UVEditorCommands.cpp
@@ -20,24 +20,24 @@ FUVEditorCommands::FUVEditorCommands()
 void FUVEditorCommands::RegisterCommands()
 {
 	// These are part of the asset editor UI
-	UI_COMMAND(OpenUVEditor, "Open UV Editor", "Open the UV Editor window", EUserInterfaceActionType::Button, FInputChord());
-	UI_COMMAND(ApplyChanges, "Apply Changes", "Apply changes without yet saving the assets", EUserInterfaceActionType::Button, FInputChord());
+	UI_COMMAND(OpenUVEditor, "UV Editor", "Open the UV Editor window.", EUserInterfaceActionType::Button, FInputChord());
+	UI_COMMAND(ApplyChanges, "Apply", "Apply changes to original meshes", EUserInterfaceActionType::Button, FInputChord());
 
 	// These get linked to various tool buttons.
-	UI_COMMAND(BeginSelectTool, "Edit", "Switch to selection-based edit tool", EUserInterfaceActionType::ToggleButton, FInputChord());
-	UI_COMMAND(BeginLayoutTool, "Layout", "Switch to layout tool", EUserInterfaceActionType::ToggleButton, FInputChord());
-	UI_COMMAND(BeginParameterizeMeshTool, "Auto UV", "Switch to Auto UV tool", EUserInterfaceActionType::ToggleButton, FInputChord());
-	UI_COMMAND(BeginChannelEditTool, "Channel Edit", "Switch to Channel Edit tool", EUserInterfaceActionType::ToggleButton, FInputChord());
-	UI_COMMAND(BeginSeamTool, "Seam", "Switch to Seam tool", EUserInterfaceActionType::ToggleButton, FInputChord());
-	UI_COMMAND(BeginRecomputeUVsTool, "Unwrap", "Switch to Unwrap tool", EUserInterfaceActionType::ToggleButton, FInputChord());
+	UI_COMMAND(BeginSelectTool, "Edit", "Selection-based UV editing", EUserInterfaceActionType::ToggleButton, FInputChord());
+	UI_COMMAND(BeginLayoutTool, "Layout", "Transform and pack existing UVs", EUserInterfaceActionType::ToggleButton, FInputChord());
+	UI_COMMAND(BeginParameterizeMeshTool, "AutoUV", "Auto-unwrap and pack UVs", EUserInterfaceActionType::ToggleButton, FInputChord());
+	UI_COMMAND(BeginChannelEditTool, "Channels", "Modify UV channels", EUserInterfaceActionType::ToggleButton, FInputChord());
+	UI_COMMAND(BeginSeamTool, "Seam", "Add UV seams", EUserInterfaceActionType::ToggleButton, FInputChord());
+	UI_COMMAND(BeginRecomputeUVsTool, "Unwrap", "Perform UV unwrapping", EUserInterfaceActionType::ToggleButton, FInputChord());
 
 	// These currently get linked to actions inside the select tool, but will eventually have their own buttons among the tools
 	// once selection is pulled out to mode-level.
-	UI_COMMAND(SewAction, "Sew", "Sew edges. The red edges will be sewn to the green edges", EUserInterfaceActionType::Button, FInputChord());
-	UI_COMMAND(SplitAction, "Split", 
-		"Given an edge selection, split those edges. Given a vertex selection, split any selected bowtie verts.", 
-		EUserInterfaceActionType::Button, FInputChord());
-	UI_COMMAND(IslandConformalUnwrapAction, "IsdUw", "Apply a conformal unwrap to the selected UV islands", EUserInterfaceActionType::Button, FInputChord());
+	UI_COMMAND(SewAction, "Sew", "Sew edges highlighted in red to edges highlighted in green", EUserInterfaceActionType::Button, FInputChord());
+	UI_COMMAND(SplitAction, "Split",
+	           "Given an edge selection, split those edges. Given a vertex selection, split any selected bowtie vertices.",
+	           EUserInterfaceActionType::Button, FInputChord());
+	UI_COMMAND(IslandConformalUnwrapAction, "IslUnw", "Apply a conformal unwrap to selected UV islands", EUserInterfaceActionType::Button, FInputChord());
 
 	// These allow us to link up to pressed keys
 	UI_COMMAND(AcceptOrCompleteActiveTool, "Accept", "Accept the active tool", EUserInterfaceActionType::Button, FInputChord(EKeys::Enter));
@@ -47,12 +47,13 @@ void FUVEditorCommands::RegisterCommands()
 	UI_COMMAND(VertexSelection, "Vertex Selection", "Select vertices", EUserInterfaceActionType::ToggleButton, FInputChord(EKeys::One));
 	UI_COMMAND(EdgeSelection, "Edge Selection", "Select edges", EUserInterfaceActionType::ToggleButton, FInputChord(EKeys::Two));
 	UI_COMMAND(TriangleSelection, "Triangle Selection", "Select triangles", EUserInterfaceActionType::ToggleButton, FInputChord(EKeys::Three));
-	UI_COMMAND(IslandSelection, "Island Selection", "Select connected islands", EUserInterfaceActionType::ToggleButton, FInputChord(EKeys::Four));
-	UI_COMMAND(FullMeshSelection, "Full Mesh Selection", "Select whole mesh", EUserInterfaceActionType::ToggleButton, FInputChord(EKeys::Five));
+	UI_COMMAND(IslandSelection, "Island Selection", "Select islands", EUserInterfaceActionType::ToggleButton, FInputChord(EKeys::Four));
+	UI_COMMAND(FullMeshSelection, "Mesh Selection", "Select meshes", EUserInterfaceActionType::ToggleButton, FInputChord(EKeys::Five));
 
-	UI_COMMAND(EnableOrbitCamera, "Orbit", "Enable Orbit Camera", EUserInterfaceActionType::ToggleButton, FInputChord());
-	UI_COMMAND(EnableFlyCamera, "Fly", "Enable Fly Camera", EUserInterfaceActionType::ToggleButton, FInputChord());
-	UI_COMMAND(ToggleBackground, "Toggle Background", "Toggle the display of the background", EUserInterfaceActionType::ToggleButton, FInputChord(EModifierKey::Alt, EKeys::B));
+	UI_COMMAND(EnableOrbitCamera, "Orbit", "Enable orbit camera", EUserInterfaceActionType::ToggleButton, FInputChord());
+	UI_COMMAND(EnableFlyCamera, "Fly", "Enable fly camera", EUserInterfaceActionType::ToggleButton, FInputChord());
+	UI_COMMAND(ToggleBackground, "Toggle Background", "Toggle background display", EUserInterfaceActionType::ToggleButton,
+	           FInputChord(EModifierKey::Alt, EKeys::B));
 }
 
 
diff --git a/Engine/Plugins/Experimental/UVEditor/Source/UVEditor/Private/UVEditorToolkit.cpp b/Engine/Plugins/Experimental/UVEditor/Source/UVEditor/Private/UVEditorToolkit.cpp
--- a/Engine/Plugins/Experimental/UVEditor/Source/UVEditor/Private/UVEditorToolkit.cpp
+++ b/Engine/Plugins/Experimental/UVEditor/Source/UVEditor/Private/UVEditorToolkit.cpp
@@ -82,7 +82,7 @@ FUVEditorToolkit::FUVEditorToolkit(UAssetEditor* InOwningAssetEditor)
 		);
 		
 	// Add any extenders specified by the UStaticMeshEditorUISubsystem
-    // The extenders provide defined locations for FModeTookit to attach
+    // The extenders provide defined locations for FModeToolkit to attach
     // tool palette tabs and detail panel tabs
 	LayoutExtender = MakeShared<FLayoutExtender>();
 	FUVEditorModule* UVEditorModule = &FModuleManager::LoadModuleChecked<FUVEditorModule>("UVEditor");
@@ -213,7 +213,6 @@ bool FUVEditorToolkit::OnRequestClose()
 		return true; 
 	}
 
-	bool bAllowClose = true;
 	// Warn the user of any unapplied changes.
 	if (UVMode->HaveUnappliedChanges())
 	{
@@ -227,17 +226,14 @@ bool FUVEditorToolkit::OnRequestClose()
 		{
 		case EAppReturnType::Yes:
 			UVMode->ApplyChanges();
-			bAllowClose = true;
 			break;
 
 		case EAppReturnType::No:
 			// exit
-			bAllowClose = true;
 			break;
 
 		case EAppReturnType::Cancel:
 			// don't exit
-			bAllowClose = false;
 			return false;
 		}
 	}
@@ -246,10 +242,7 @@ bool FUVEditorToolkit::OnRequestClose()
     // This is super important to do, otherwise currently opened tabs won't be marked as "closed".
     // This results in tabs not being properly recycled upon reopening the editor and tab
     // duplication for each opening event.
-	if (bAllowClose)
-	{
-		GetEditorModeManager().ActivateDefaultMode();
-	}
+	GetEditorModeManager().ActivateDefaultMode();
 
 	return FAssetEditorToolkit::OnRequestClose();
 }
@@ -433,7 +426,7 @@ void FUVEditorToolkit::PostInitAssetEditor()
 			return UVModeToolkit->CreateChannelMenu();
 		}),
 		LOCTEXT("UVEditorChannelMenu_Label", "Channels"),
-		LOCTEXT("UVEditorChannelMenu_ToolTip", "Change the selected UV Channel for each asset."),
+		LOCTEXT("UVEditorChannelMenu_ToolTip", "Select the current UV Channel for each mesh"),
 		FSlateIcon(FUVEditorStyle::Get().GetStyleSetName(), "UVEditor.ChannelSettings")
 	));
 
@@ -446,7 +439,7 @@ void FUVEditorToolkit::PostInitAssetEditor()
 			return UVModeToolkit->CreateBackgroundSettingsWidget();
 		}),
 		LOCTEXT("UVEditorBackgroundSettings_Label", "Display"),
-		LOCTEXT("UVEditorBackgroundSettings_ToolTip", "Change the background and other viewport display settings."),
+		LOCTEXT("UVEditorBackgroundSettings_ToolTip", "Change the background display settings"),
 		FSlateIcon(FUVEditorStyle::Get().GetStyleSetName(), "UVEditor.BackgroundSettings")
 	));
 
diff --git a/Engine/Plugins/Experimental/UVEditor/Source/UVEditorTools/Private/UVEditorLayerEditTool.cpp b/Engine/Plugins/Experimental/UVEditor/Source/UVEditorTools/Private/UVEditorLayerEditTool.cpp
--- a/Engine/Plugins/Experimental/UVEditor/Source/UVEditorTools/Private/UVEditorLayerEditTool.cpp
+++ b/Engine/Plugins/Experimental/UVEditor/Source/UVEditorTools/Private/UVEditorLayerEditTool.cpp
@@ -22,7 +22,7 @@ using namespace UE::Geometry;
 namespace UUVEditorChannelEditLocals
 {
 	const FText UVChannelAddTransactionName = LOCTEXT("UVChannelAddTransactionName", "Add UV Channel");
-	const FText UVChannelCloneTransactionName = LOCTEXT("UVChannelCloneTransactionName", "Clone UV Channel");
+	const FText UVChannelCopyTransactionName = LOCTEXT("UVChannelCopyTransactionName", "Copy UV Channel");
 	const FText UVChannelDeleteTransactionName = LOCTEXT("UVChannelDeleteTransactionName", "Delete UV Channel");
 
 	void DeleteChannel(UUVEditorToolMeshInput* Target, int32 DeletedUVChannelIndex, 
@@ -119,10 +119,11 @@ namespace UUVEditorChannelEditLocals
 		int32 AddedUVChannelIndex;		
 	};
 
-	class FInputObjectUVChannelClone : public FToolCommandChange
+	class FInputObjectUVChannelCopy : public FToolCommandChange
 	{
 	public:
-		FInputObjectUVChannelClone(const TObjectPtr<UUVEditorToolMeshInput>& TargetIn, int32 SourceUVChannelIndexIn, int32 TargetUVChannelIndexIn, const FDynamicMeshUVOverlay& OriginalUVChannelIn)
+		FInputObjectUVChannelCopy(const TObjectPtr<UUVEditorToolMeshInput>& TargetIn, int32 SourceUVChannelIndexIn, int32 TargetUVChannelIndexIn,
+		                          const FDynamicMeshUVOverlay& OriginalUVChannelIn)
 			: Target(TargetIn)			
 			, SourceUVChannelIndex(SourceUVChannelIndexIn)
 			, TargetUVChannelIndex(TargetUVChannelIndexIn)
@@ -137,7 +138,6 @@ namespace UUVEditorChannelEditLocals
 			{
 				return;
 			}
-			UUVToolAssetAndChannelAPI* AssetAndChannelAPI = InteractiveToolManager->GetContextObjectStore()->FindContext<UUVToolAssetAndChannelAPI>();
 
 			FDynamicMeshUVEditor DynamicMeshUVEditor(Target->AppliedCanonical.Get(), TargetUVChannelIndex, false);
 			DynamicMeshUVEditor.CopyUVLayer(Target->AppliedCanonical->Attributes()->GetUVLayer(SourceUVChannelIndex));
@@ -151,7 +151,6 @@ namespace UUVEditorChannelEditLocals
 			{
 				return;
 			}
-			UUVToolAssetAndChannelAPI* AssetAndChannelAPI = InteractiveToolManager->GetContextObjectStore()->FindContext<UUVToolAssetAndChannelAPI>();
 
 			FDynamicMeshUVEditor DynamicMeshUVEditor(Target->AppliedCanonical.Get(), TargetUVChannelIndex, false);
 			DynamicMeshUVEditor.CopyUVLayer(&OriginalUVChannel);
@@ -165,7 +164,7 @@ namespace UUVEditorChannelEditLocals
 
 		virtual FString ToString() const override
 		{
-			return TEXT("UVEditorModeLocals::FInputObjectUVChannelClone");
+			return TEXT("UVEditorModeLocals::FInputObjectUVChannelCopy");
 		}
 
 	protected:
@@ -375,8 +374,7 @@ bool UUVEditorChannelEditTargetProperties::ValidateUVChannelSelection(bool bUpda
 		return true;
 	};
 
-	bValid = CheckAndUpdateChannel(TargetChannel);
-	bValid = bValid && CheckAndUpdateChannel(TargetChannel);
+	bValid &= CheckAndUpdateChannel(TargetChannel);
 
 	return bValid;
 }
@@ -490,7 +488,7 @@ void UUVEditorChannelEditTool::Setup()
 	AddActionProperties = NewObject<UUVEditorChannelEditAddProperties>(this);
 	AddToolPropertySource(AddActionProperties);
 
-	CopyActionProperties = NewObject<UUVEditorChannelEditCloneProperties>(this);
+	CopyActionProperties = NewObject<UUVEditorChannelEditCopyProperties>(this);
 	AddToolPropertySource(CopyActionProperties);
 
 	DeleteActionProperties = NewObject<UUVEditorChannelEditDeleteProperties>(this);
@@ -501,7 +499,7 @@ void UUVEditorChannelEditTool::Setup()
 	AddToolPropertySource(ToolActions);
 
 	SetToolDisplayName(LOCTEXT("ToolName", "UV Channel Edit"));
-	GetToolManager()->DisplayMessage(LOCTEXT("OnStartUVChannelEditTool", "Add/Copy/Delete Asset UV Channels"),
+	GetToolManager()->DisplayMessage(LOCTEXT("OnStartUVChannelEditTool", "Add, copy or delete UV channels"),
 		EToolMessageLevel::UserNotification);
 
 	for (int32 i = 0; i < Targets.Num(); ++i)
@@ -641,12 +639,10 @@ void UUVEditorChannelEditTool::AddChannel()
 
 void UUVEditorChannelEditTool::CopyChannel()
 {
-	UUVToolAssetAndChannelAPI* AssetAndChannelAPI = GetToolManager()->GetContextObjectStore()->FindContext<UUVToolAssetAndChannelAPI>();
-
 	EmitChangeAPI->EmitToolIndependentChange(this,
-		MakeUnique<UUVEditorChannelEditLocals::FInputObjectUVChannelClone>(Targets[ActiveAsset], ReferenceChannel, ActiveChannel, 
+		MakeUnique<UUVEditorChannelEditLocals::FInputObjectUVChannelCopy>(Targets[ActiveAsset], ReferenceChannel, ActiveChannel, 
 			*Targets[ActiveAsset]->AppliedCanonical->Attributes()->GetUVLayer(ActiveChannel)),
-		UUVEditorChannelEditLocals::UVChannelCloneTransactionName);
+		UUVEditorChannelEditLocals::UVChannelCopyTransactionName);
 
 	FDynamicMeshUVEditor DynamicMeshUVEditor(Targets[ActiveAsset]->AppliedCanonical.Get(), ActiveChannel, false);
 	DynamicMeshUVEditor.CopyUVLayer(Targets[ActiveAsset]->AppliedCanonical->Attributes()->GetUVLayer(ReferenceChannel));
diff --git a/Engine/Plugins/Experimental/UVEditor/Source/UVEditorTools/Private/UVEditorLayoutTool.cpp b/Engine/Plugins/Experimental/UVEditor/Source/UVEditorTools/Private/UVEditorLayoutTool.cpp
--- a/Engine/Plugins/Experimental/UVEditor/Source/UVEditorTools/Private/UVEditorLayoutTool.cpp
+++ b/Engine/Plugins/Experimental/UVEditor/Source/UVEditorTools/Private/UVEditorLayoutTool.cpp
@@ -62,7 +62,7 @@ void UUVEditorLayoutTool::Setup()
 	}
 
 	SetToolDisplayName(LOCTEXT("ToolName", "UV Layout"));
-	GetToolManager()->DisplayMessage(LOCTEXT("OnStartUVLayoutTool", "Transform/Rotate/Scale existing UV Charts using various strategies"),
+	GetToolManager()->DisplayMessage(LOCTEXT("OnStartUVLayoutTool", "Translate, rotate or scale existing UV Charts using various strategies"),
 		EToolMessageLevel::UserNotification);
 }
 
diff --git a/Engine/Plugins/Experimental/UVEditor/Source/UVEditorTools/Private/UVEditorRecomputeUVsTool.cpp b/Engine/Plugins/Experimental/UVEditor/Source/UVEditorTools/Private/UVEditorRecomputeUVsTool.cpp
--- a/Engine/Plugins/Experimental/UVEditor/Source/UVEditorTools/Private/UVEditorRecomputeUVsTool.cpp
+++ b/Engine/Plugins/Experimental/UVEditor/Source/UVEditorTools/Private/UVEditorRecomputeUVsTool.cpp
@@ -5,14 +5,11 @@
 #include "ToolBuilderUtil.h"
 #include "DynamicMesh/DynamicMesh3.h"
 #include "Polygroups/PolygroupUtil.h"
-#include "FaceGroupUtil.h"
-#include "ToolSetupUtil.h"
-#include "ModelingToolTargetUtil.h"
 #include "ParameterizationOps/RecomputeUVsOp.h"
 #include "ToolTargets/UVEditorToolMeshInput.h"
-#include "Properties/MeshMaterialProperties.h"
 #include "UVToolContextObjects.h"
 #include "ContextObjectStore.h"
+#include "MeshOpPreviewHelpers.h"
 
 using namespace UE::Geometry;
 
@@ -83,14 +80,13 @@ void UUVEditorRecomputeUVsTool::Setup()
 
 	SetToolDisplayName(LOCTEXT("ToolNameLocal", "UV Unwrap"));
 	GetToolManager()->DisplayMessage(
-		LOCTEXT("OnStartTool_Regions", "Generate UVs for Polygroups or existing UV charts of the Mesh using various strategies."),
+		LOCTEXT("OnStartTool_Regions", "Generate UVs for PolyGroups or existing UV islands of the mesh using various strategies."),
 		EToolMessageLevel::UserNotification);
 }
 
 
 void UUVEditorRecomputeUVsTool::OnPropertyModified(UObject* PropertySet, FProperty* Property)
 {
-	bool bForceMaterialUpdate = false;
 	if (PropertySet == Settings )
 	{
 		// One of the UV generation properties must have changed.  Dirty the result to force a recompute
@@ -98,9 +94,7 @@ void UUVEditorRecomputeUVsTool::OnPropertyModified(UObject* PropertySet, FProper
 		{
 			Target->AppliedPreview->InvalidateResult();
 		}
-		bForceMaterialUpdate = true;
 	}
-
 }
 
 
@@ -195,7 +189,7 @@ void UUVEditorRecomputeUVsTool::UpdateActiveGroupLayer()
 		{
 			FName SelectedName = PolygroupLayerProperties->ActiveGroupLayer;
 			FDynamicMeshPolygroupAttribute* FoundAttrib = UE::Geometry::FindPolygroupLayerByName(*Targets[0]->AppliedCanonical, SelectedName);
-			ensureMsgf(FoundAttrib, TEXT("Selected Attribute Not Found! Falling back to Default group layer."));
+			ensureMsgf(FoundAttrib, TEXT("Selected attribute not found! Falling back to Default group layer."));
 			ActiveGroupSet = MakeShared<UE::Geometry::FPolygroupSet, ESPMode::ThreadSafe>(Targets[0]->AppliedCanonical.Get(), FoundAttrib);
 		}
 	}
diff --git a/Engine/Plugins/Experimental/UVEditor/Source/UVEditorTools/Private/UVEditorSeamTool.cpp b/Engine/Plugins/Experimental/UVEditor/Source/UVEditorTools/Private/UVEditorSeamTool.cpp
--- a/Engine/Plugins/Experimental/UVEditor/Source/UVEditorTools/Private/UVEditorSeamTool.cpp
+++ b/Engine/Plugins/Experimental/UVEditor/Source/UVEditorTools/Private/UVEditorSeamTool.cpp
@@ -477,7 +477,7 @@ void UUVEditorSeamTool::Setup()
 		});
 	}
 
-	GetToolManager()->DisplayMessage(LOCTEXT("OnStartSeamTool", "Click to add seams in 2d/3d view. Click start/end point or hit Enter to confirm."),
+	GetToolManager()->DisplayMessage(LOCTEXT("OnStartSeamTool", "Click to add seams in 2D or 3D viewport. To finish the seam, either click on the last point added, press enter or create a loop by clicking on the first point."),
 		EToolMessageLevel::UserNotification);
 }
 
diff --git a/Engine/Plugins/Experimental/UVEditor/Source/UVEditorTools/Private/UVSelectTool.cpp b/Engine/Plugins/Experimental/UVEditor/Source/UVEditorTools/Private/UVSelectTool.cpp
--- a/Engine/Plugins/Experimental/UVEditor/Source/UVEditorTools/Private/UVSelectTool.cpp
+++ b/Engine/Plugins/Experimental/UVEditor/Source/UVEditorTools/Private/UVSelectTool.cpp
@@ -396,7 +396,7 @@ void UUVSelectTool::Setup()
 	// The trees will be updated by the tree store, which listens to the same broadcasts.
 	for (int32 i = 0; i < Targets.Num(); ++i)
 	{
-		Targets[i]->OnCanonicalModified.AddWeakLambda(this, [this, i]
+		Targets[i]->OnCanonicalModified.AddWeakLambda(this, [this]
 		(UUVEditorToolMeshInput* InputObject, const UUVEditorToolMeshInput::FCanonicalModifiedInfo& Info) {
 			if (bIgnoreOnCanonicalChange) // Used to avoid reacting to broadcasts that we ourselves caused
 			{
@@ -464,7 +464,7 @@ void UUVSelectTool::Setup()
 	UpdateGizmo();
 
 	GetToolManager()->DisplayMessage(LOCTEXT("SelectToolStatusBarMessage", 
-		"Select elements in the viewport and then use one of the edit action buttons."), 
+		"Select elements in the viewport and then transform them or use one of the action buttons."), 
 		EToolMessageLevel::UserNotification);
 }
 
@@ -935,7 +935,7 @@ void UUVSelectTool::ApplyAction(ESelectToolAction ActionType)
 			EmitChangeAPI->BeginUndoTransaction(TransactionName);
 
 			SelectionMechanic->SetSelection(FDynamicMeshSelection(), false, true);
-			bool ActionSuccessful = SewAction->ExecuteAction(*EmitChangeAPI);
+			SewAction->ExecuteAction(*EmitChangeAPI);
 
 			EmitChangeAPI->EndUndoTransaction();
 		}
@@ -949,7 +949,7 @@ void UUVSelectTool::ApplyAction(ESelectToolAction ActionType)
 			EmitChangeAPI->BeginUndoTransaction(TransactionName);
 
 			SelectionMechanic->SetSelection(FDynamicMeshSelection(), false, true);
-			bool ActionSuccessful = IslandConformalUnwrapAction->ExecuteAction(*EmitChangeAPI);
+			IslandConformalUnwrapAction->ExecuteAction(*EmitChangeAPI);
 
 			EmitChangeAPI->EndUndoTransaction();
 		}
diff --git a/Engine/Plugins/Experimental/UVEditor/Source/UVEditorTools/Public/UVEditorLayerEditTool.h b/Engine/Plugins/Experimental/UVEditor/Source/UVEditorTools/Public/UVEditorLayerEditTool.h
--- a/Engine/Plugins/Experimental/UVEditor/Source/UVEditorTools/Public/UVEditorLayerEditTool.h
+++ b/Engine/Plugins/Experimental/UVEditor/Source/UVEditorTools/Public/UVEditorLayerEditTool.h
@@ -128,7 +128,7 @@ public:
 };
 
 UCLASS()
-class UVEDITORTOOLS_API UUVEditorChannelEditCloneProperties : public UInteractiveToolPropertySet
+class UVEDITORTOOLS_API UUVEditorChannelEditCopyProperties : public UInteractiveToolPropertySet
 {
 	GENERATED_BODY()
 public:
@@ -220,7 +220,7 @@ protected:
 	TObjectPtr<UUVEditorChannelEditAddProperties> AddActionProperties = nullptr;
 
 	UPROPERTY()
-	TObjectPtr<UUVEditorChannelEditCloneProperties> CopyActionProperties = nullptr;
+	TObjectPtr<UUVEditorChannelEditCopyProperties> CopyActionProperties = nullptr;
 
 	UPROPERTY()
 	TObjectPtr<UUVEditorChannelEditDeleteProperties> DeleteActionProperties = nullptr;
diff --git a/Engine/Plugins/Experimental/UVEditor/Source/UVEditorTools/Public/UVEditorRecomputeUVsTool.h b/Engine/Plugins/Experimental/UVEditor/Source/UVEditorTools/Public/UVEditorRecomputeUVsTool.h
--- a/Engine/Plugins/Experimental/UVEditor/Source/UVEditorTools/Public/UVEditorRecomputeUVsTool.h
+++ b/Engine/Plugins/Experimental/UVEditor/Source/UVEditorTools/Public/UVEditorRecomputeUVsTool.h
@@ -5,7 +5,6 @@
 #include "CoreMinimal.h"
 #include "BaseTools/SingleSelectionMeshEditingTool.h"
 #include "DynamicMesh/DynamicMesh3.h"
-#include "MeshOpPreviewHelpers.h"
 #include "PropertySets/PolygroupLayersProperties.h"
 #include "Properties/RecomputeUVsProperties.h"
 #include "Polygroups/PolygroupSet.h"
@@ -14,7 +13,7 @@
 #include "UVEditorRecomputeUVsTool.generated.h"
 
 
-// predeclarations
+// Forward declarations
 class UDynamicMeshComponent;
 class UMaterialInterface;
 class UMaterialInstanceDynamic;
diff --git a/Engine/Plugins/Runtime/MeshModelingToolset/Source/MeshModelingTools/Private/RecomputeUVsTool.cpp b/Engine/Plugins/Runtime/MeshModelingToolset/Source/MeshModelingTools/Private/RecomputeUVsTool.cpp
--- a/Engine/Plugins/Runtime/MeshModelingToolset/Source/MeshModelingTools/Private/RecomputeUVsTool.cpp
+++ b/Engine/Plugins/Runtime/MeshModelingToolset/Source/MeshModelingTools/Private/RecomputeUVsTool.cpp
@@ -110,7 +110,7 @@ void URecomputeUVsTool::Setup()
 
 	SetToolDisplayName(LOCTEXT("ToolNameLocal", "UV Unwrap"));
 	GetToolManager()->DisplayMessage(
-		LOCTEXT("OnStartTool_Regions", "Generate UVs for Polygroups or existing UV charts of the Mesh using various strategies."),
+		LOCTEXT("OnStartTool_Regions", "Generate UVs for PolyGroups or existing UV islands of the mesh using various strategies."),
 		EToolMessageLevel::UserNotification);
 }
 
