UVEditor: Highlight edge and vertex selections in the live preview.

#rnx

diff --git a/Engine/Plugins/Experimental/UVEditor/Source/UVEditorTools/Private/UVEditorUXSettings.cpp b/Engine/Plugins/Experimental/UVEditor/Source/UVEditorTools/Private/UVEditorUXSettings.cpp
--- a/Engine/Plugins/Experimental/UVEditor/Source/UVEditorTools/Private/UVEditorUXSettings.cpp
+++ b/Engine/Plugins/Experimental/UVEditor/Source/UVEditorTools/Private/UVEditorUXSettings.cpp
@@ -50,6 +50,7 @@ const FColor FUVEditorUXSettings::GridMinorColor(FColor::FromHex("#777777"));
 
 // Thicknesses
 const float FUVEditorUXSettings::LivePreviewHighlightThickness(2.0);
+const float FUVEditorUXSettings::LivePreviewHighlightPointSize(4);
 const float FUVEditorUXSettings::SelectionLineThickness(1.5);
 const float FUVEditorUXSettings::SelectionPointThickness(6);
 const float FUVEditorUXSettings::SewLineHighlightThickness(3.0f);
diff --git a/Engine/Plugins/Experimental/UVEditor/Source/UVEditorTools/Private/UVSelectTool.cpp b/Engine/Plugins/Experimental/UVEditor/Source/UVEditorTools/Private/UVSelectTool.cpp
--- a/Engine/Plugins/Experimental/UVEditor/Source/UVEditorTools/Private/UVSelectTool.cpp
+++ b/Engine/Plugins/Experimental/UVEditor/Source/UVEditorTools/Private/UVSelectTool.cpp
@@ -8,6 +8,7 @@
 #include "ContextObjectStore.h"
 #include "Drawing/LineSetComponent.h"
 #include "Drawing/MeshElementsVisualizer.h"
+#include "Drawing/PointSetComponent.h"
 #include "Drawing/PreviewGeometryActor.h"
 #include "DynamicMesh/DynamicMeshChangeTracker.h"
 #include "DynamicMesh/MeshIndexUtil.h"
@@ -447,6 +448,12 @@ void UUVSelectTool::Setup()
 	LivePreviewLineSet->SetLineMaterial(ToolSetupUtil::GetDefaultLineComponentMaterial(
 		GetToolManager(), /*bDepthTested*/ true));
 
+	LivePreviewPointSet = NewObject<UPointSetComponent>(LivePreviewGeometryActor);
+	LivePreviewPointSet->AttachToComponent(LivePreviewLineSet, FAttachmentTransformRules::KeepWorldTransform);
+	LivePreviewPointSet->RegisterComponent();
+	LivePreviewPointSet->SetPointMaterial(ToolSetupUtil::GetDefaultPointComponentMaterial(
+		GetToolManager(), /*bDepthTested*/ true));
+
 	SewAction = NewObject<UUVSeamSewAction>();
 	SewAction->Setup(this);
 	SewAction->SetTargets(Targets);
@@ -492,6 +499,8 @@ void UUVSelectTool::Shutdown(EToolShutdownType ShutdownType)
 	{
 		LivePreviewGeometryActor->Destroy();
 		LivePreviewGeometryActor = nullptr;
+		LivePreviewPointSet = nullptr;
+		LivePreviewLineSet = nullptr;
 	}
 
 	if (SewAction)
@@ -631,7 +640,8 @@ void UUVSelectTool::OnSelectionChanged()
 	SelectionTargetIndex = -1;
 	MovingVids.Reset();
 	SelectedTids.Reset();
-	LivePreviewBoundaryEids.Reset();
+	LivePreviewEids.Reset();
+	LivePreviewVids.Reset();
 
 	if (!Selection.IsEmpty())
 	{
@@ -646,6 +656,9 @@ void UUVSelectTool::OnSelectionChanged()
 		}
 		check(SelectionTargetIndex >= 0);
 
+		UUVEditorToolMeshInput* Target = Targets[SelectionTargetIndex];
+		const FDynamicMesh3* LivePreviewMesh = Target->AppliedCanonical.Get();
+
 		// Note the selected vids
 		TSet<int32> VidSet;
 		TSet<int32> TidSet;
@@ -653,7 +666,6 @@ void UUVSelectTool::OnSelectionChanged()
 		{
 			TRACE_CPUPROFILER_EVENT_SCOPE(Triangle);
 
-			const FDynamicMesh3* LivePreviewMesh = Targets[SelectionTargetIndex]->AppliedCanonical.Get();
 			for (int32 Tid : Selection.SelectedIDs)
 			{
 				FIndex3i TriVids = Selection.Mesh->GetTriangle(Tid);
@@ -680,7 +692,7 @@ void UUVSelectTool::OnSelectionChanged()
 					{
 						if (EdgeTids[j] != Tid && !Selection.SelectedIDs.Contains(EdgeTids[j]))
 						{
-							LivePreviewBoundaryEids.Add(TriEids[i]);
+							LivePreviewEids.Add(TriEids[i]);
 							break;
 						}
 					}
@@ -713,6 +725,11 @@ void UUVSelectTool::OnSelectionChanged()
 						}
 					}
 				}
+
+				// Add the edge highlight in the live preview
+				LivePreviewEids.Add(LivePreviewMesh->FindEdge(
+					Target->UnwrapVidToAppliedVid(EdgeVids.A),
+					Target->UnwrapVidToAppliedVid(EdgeVids.B)));
 			}
 		}
 		else if (Selection.Type == FDynamicMeshSelection::EType::Vertex)
@@ -738,6 +755,7 @@ void UUVSelectTool::OnSelectionChanged()
 					}
 				}
 
+				LivePreviewVids.Add(Target->UnwrapVidToAppliedVid(Vid));
 			}
 		}
 		else
@@ -763,6 +781,7 @@ void UUVSelectTool::UpdateLivePreviewLines()
 	TRACE_CPUPROFILER_EVENT_SCOPE(UVSelectTool_UpdateLivePreviewLines);
 	
 	LivePreviewLineSet->Clear();
+	LivePreviewPointSet->Clear();
 
 	const FDynamicMeshSelection& Selection = SelectionMechanic->GetCurrentSelection();
 	if (!Selection.IsEmpty())
@@ -770,7 +789,7 @@ void UUVSelectTool::UpdateLivePreviewLines()
 		FTransform MeshTransform = Targets[SelectionTargetIndex]->AppliedPreview->PreviewMesh->GetTransform();
 		const FDynamicMesh3* LivePreviewMesh = Targets[SelectionTargetIndex]->AppliedCanonical.Get();
 
-		for (int32 Eid : LivePreviewBoundaryEids)
+		for (int32 Eid : LivePreviewEids)
 		{
 			FVector3d Vert1, Vert2;
 			LivePreviewMesh->GetEdgeV(Eid, Vert1, Vert2);
@@ -781,7 +800,17 @@ void UUVSelectTool::UpdateLivePreviewLines()
 				FUVEditorUXSettings::SelectionTriangleWireframeColor,
 				FUVEditorUXSettings::LivePreviewHighlightThickness,
 				FUVEditorUXSettings::LivePreviewHighlightDepthOffset);
-		}	
+		}
+
+		for (int32 Vid : LivePreviewVids)
+		{
+			FVector3d Position = LivePreviewMesh->GetVertex(Vid);
+
+			LivePreviewPointSet->AddPoint(Position, 
+				FUVEditorUXSettings::SelectionTriangleWireframeColor,
+				FUVEditorUXSettings::LivePreviewHighlightPointSize,
+				FUVEditorUXSettings::LivePreviewHighlightDepthOffset);
+		}
 	}
 }
 
diff --git a/Engine/Plugins/Experimental/UVEditor/Source/UVEditorTools/Public/UVEditorUXSettings.h b/Engine/Plugins/Experimental/UVEditor/Source/UVEditorTools/Public/UVEditorUXSettings.h
--- a/Engine/Plugins/Experimental/UVEditor/Source/UVEditorTools/Public/UVEditorUXSettings.h
+++ b/Engine/Plugins/Experimental/UVEditor/Source/UVEditorTools/Public/UVEditorUXSettings.h
@@ -46,6 +46,7 @@ public:
 	static const FColor SelectionHoverTriangleWireframeColor;
 
 	static const float LivePreviewHighlightThickness;
+	static const float LivePreviewHighlightPointSize;
 	static const float LivePreviewHighlightDepthOffset;
 
 	static const float SelectionLineThickness;
diff --git a/Engine/Plugins/Experimental/UVEditor/Source/UVEditorTools/Public/UVSelectTool.h b/Engine/Plugins/Experimental/UVEditor/Source/UVEditorTools/Public/UVSelectTool.h
--- a/Engine/Plugins/Experimental/UVEditor/Source/UVEditorTools/Public/UVSelectTool.h
+++ b/Engine/Plugins/Experimental/UVEditor/Source/UVEditorTools/Public/UVSelectTool.h
@@ -20,6 +20,7 @@ class APreviewGeometryActor;
 class FToolCommandChange;
 class ULineSetComponent;
 class UMeshSelectionMechanic;
+class UPointSetComponent;
 class UToolTargetManager;
 class UCombinedTransformGizmo;
 class UTransformProxy;
@@ -168,6 +169,9 @@ protected:
 	UPROPERTY()
 	TObjectPtr<ULineSetComponent> LivePreviewLineSet = nullptr;
 
+	UPROPERTY()
+	TObjectPtr<UPointSetComponent> LivePreviewPointSet = nullptr;
+
 	UPROPERTY()
 	TObjectPtr<UUVToolEmitChangeAPI> EmitChangeAPI = nullptr;
 
@@ -191,7 +195,8 @@ protected:
 	TArray<int32> SelectedTids;
 	TArray<FVector3d> MovingVertOriginalPositions;
 	int32 SelectionTargetIndex;
-	TArray<int32> LivePreviewBoundaryEids;
+	TSet<int32> LivePreviewEids;
+	TSet<int32> LivePreviewVids;
 	
 	// When selecting edges, used to hold the edges as pairs of vids, because the eids change
 	// during undo/redo and other topological operations.
diff --git a/Engine/Plugins/Runtime/MeshModelingToolset/Source/ModelingComponents/Public/Drawing/PointSetComponent.h b/Engine/Plugins/Runtime/MeshModelingToolset/Source/ModelingComponents/Public/Drawing/PointSetComponent.h
--- a/Engine/Plugins/Runtime/MeshModelingToolset/Source/ModelingComponents/Public/Drawing/PointSetComponent.h
+++ b/Engine/Plugins/Runtime/MeshModelingToolset/Source/ModelingComponents/Public/Drawing/PointSetComponent.h
@@ -61,9 +61,16 @@ public:
 	/** Reserve enough memory for up to the given ID */
 	void ReservePoints(const int32 MaxID);
 
-	/** Add a point to the set */
+	/** Add a point to be rendered using the component. */
 	int32 AddPoint(const FRenderablePoint& OverlayPoint);
 
+	/** Create and add a point to be rendered using the component. */
+	int32 AddPoint(const FVector& InPosition, const FColor& InColor, const float InSize, const float InDepthBias = 0.0f)
+	{
+		// This is just a convenience function to avoid client code having to know about FRenderablePoint.
+		return AddPoint(FRenderablePoint(InPosition, InColor, InSize, InDepthBias));
+	}
+
 	/** Insert a point with the given ID into the set. */
 	void InsertPoint(const int32 ID, const FRenderablePoint& OverlayPoint);
 
