UVEditor: Remove "complete" button from complete-style tools. Disallow switching away for accept/cancel-style tools.

#rnx

diff --git a/Engine/Plugins/Experimental/UVEditor/Source/UVEditor/Private/UVEditorMode.cpp b/Engine/Plugins/Experimental/UVEditor/Source/UVEditor/Private/UVEditorMode.cpp
--- a/Engine/Plugins/Experimental/UVEditor/Source/UVEditor/Private/UVEditorMode.cpp
+++ b/Engine/Plugins/Experimental/UVEditor/Source/UVEditor/Private/UVEditorMode.cpp
@@ -210,6 +210,17 @@ void UUVEditorMode::RegisterTools()
 	RegisterTool(CommandInfos.BeginRecomputeUVsTool, TEXT("BeginRecomputeUVsTool"), UVEditorRecomputeUVsToolBuilder);
 }
 
+bool UUVEditorMode::ShouldToolStartBeAllowed(const FString& ToolIdentifier) const
+{
+	// For now we've decided to disallow switch-away on accept/cancel tools in the UV editor.
+	if (GetInteractiveToolsContext()->ActiveToolHasAccept())
+	{
+		return false;
+	}
+	
+	return Super::ShouldToolStartBeAllowed(ToolIdentifier);
+}
+
 void UUVEditorMode::CreateToolkit()
 {
 	Toolkit = MakeShared<FUVEditorModeToolkit>();
diff --git a/Engine/Plugins/Experimental/UVEditor/Source/UVEditor/Private/UVEditorModeToolkit.cpp b/Engine/Plugins/Experimental/UVEditor/Source/UVEditor/Private/UVEditorModeToolkit.cpp
--- a/Engine/Plugins/Experimental/UVEditor/Source/UVEditor/Private/UVEditorModeToolkit.cpp
+++ b/Engine/Plugins/Experimental/UVEditor/Source/UVEditor/Private/UVEditorModeToolkit.cpp
@@ -199,27 +199,29 @@ void FUVEditorModeToolkit::Init(const TSharedPtr<IToolkitHost>& InitToolkitHost,
 				.Visibility_Lambda([this]() { return GetScriptableEditorMode()->GetInteractiveToolsContext()->ActiveToolHasAccept() ? EVisibility::Visible : EVisibility::Collapsed; })
 			]
 
-			+ SHorizontalBox::Slot()
-				.AutoWidth()
-				.Padding(FMargin(2.0, 0.f, 0.f, 0.f))
-				[
-					SNew(SButton)
-					.ButtonStyle(FAppStyle::Get(), "PrimaryButton")
-				.TextStyle(FAppStyle::Get(), "DialogButtonText")
-				.Text(LOCTEXT("OverlayComplete", "Complete"))
-				.ToolTipText(LOCTEXT("OverlayCompleteTooltip", "Exit the active Tool [Enter]"))
-				.HAlign(HAlign_Center)
-				.OnClicked_Lambda([this]() { 
-					GetScriptableEditorMode()->GetInteractiveToolsContext()->EndTool(EToolShutdownType::Completed); 
-					Cast<UUVEditorMode>(GetScriptableEditorMode())->ActivateDefaultTool();
-					return FReply::Handled(); 
-					})
-				.IsEnabled_Lambda([this]() {
-					UUVEditorMode* Mode = Cast<UUVEditorMode>(GetScriptableEditorMode());
-					return GetScriptableEditorMode()->GetInteractiveToolsContext()->CanCompleteActiveTool();
-				})
-				.Visibility_Lambda([this]() { return GetScriptableEditorMode()->GetInteractiveToolsContext()->CanCompleteActiveTool() ? EVisibility::Visible : EVisibility::Collapsed; })
-			]
+			// For now we've decided not to use a "Complete" button for complete-style tools, instead requiring
+			// users to just select a different tool. Uncomment the below if we want to put it back.
+			//+ SHorizontalBox::Slot()
+			//.AutoWidth()
+			//.Padding(FMargin(2.0, 0.f, 0.f, 0.f))
+			//[
+			//	SNew(SButton)
+			//	.ButtonStyle(FAppStyle::Get(), "PrimaryButton")
+			//	.TextStyle(FAppStyle::Get(), "DialogButtonText")
+			//	.Text(LOCTEXT("OverlayComplete", "Complete"))
+			//	.ToolTipText(LOCTEXT("OverlayCompleteTooltip", "Exit the active Tool [Enter]"))
+			//	.HAlign(HAlign_Center)
+			//	.OnClicked_Lambda([this]() { 
+			//		GetScriptableEditorMode()->GetInteractiveToolsContext()->EndTool(EToolShutdownType::Completed); 
+			//		Cast<UUVEditorMode>(GetScriptableEditorMode())->ActivateDefaultTool();
+			//		return FReply::Handled(); 
+			//		})
+			//	.IsEnabled_Lambda([this]() {
+			//		UUVEditorMode* Mode = Cast<UUVEditorMode>(GetScriptableEditorMode());
+			//		return GetScriptableEditorMode()->GetInteractiveToolsContext()->CanCompleteActiveTool();
+			//	})
+			//	.Visibility_Lambda([this]() { return GetScriptableEditorMode()->GetInteractiveToolsContext()->CanCompleteActiveTool() ? EVisibility::Visible : EVisibility::Collapsed; })
+			//]
 		]	
 	];
 
@@ -355,9 +357,11 @@ void FUVEditorModeToolkit::OnToolStarted(UInteractiveToolManager* Manager, UInte
 		Mode->GetInteractiveToolsContext()->GetTransactionAPI()->AppendChange(
 			Mode, MakeUnique<UVEditorModeToolkitLocals::FUVEditorBeginToolChange>(), LOCTEXT("ActivateTool", "Activate Tool"));
 
-		// Add the accept/cancel overlay. Again, unless we're in the default tool, which we don't leave
-		// except by activating another tool.
-		GetToolkitHost()->AddViewportOverlayWidget(ViewportOverlayWidget.ToSharedRef());
+		if (Mode->GetInteractiveToolsContext()->ActiveToolHasAccept())
+		{
+			// Add the accept/cancel overlay only if the tool has accept/cancel.
+			GetToolkitHost()->AddViewportOverlayWidget(ViewportOverlayWidget.ToSharedRef());
+		}
 	}
 }
 
diff --git a/Engine/Plugins/Experimental/UVEditor/Source/UVEditor/Public/UVEditorMode.h b/Engine/Plugins/Experimental/UVEditor/Source/UVEditor/Public/UVEditorMode.h
--- a/Engine/Plugins/Experimental/UVEditor/Source/UVEditor/Public/UVEditorMode.h
+++ b/Engine/Plugins/Experimental/UVEditor/Source/UVEditor/Public/UVEditorMode.h
@@ -103,6 +103,7 @@ public:
 
 	// UEdMode overrides
 	virtual void Enter() override;
+	virtual bool ShouldToolStartBeAllowed(const FString& ToolIdentifier) const override;
 	virtual void Exit() override;
 	virtual void ModeTick(float DeltaTime) override;
 	// We're changing visibility of this one to public here so that we can call it from the toolkit
