UVEditor: Change tool message in seam tool based on state. Accept the pending seam when switching away.

#rnx

diff --git a/Engine/Plugins/Experimental/UVEditor/Source/UVEditorTools/Private/UVEditorSeamTool.cpp b/Engine/Plugins/Experimental/UVEditor/Source/UVEditorTools/Private/UVEditorSeamTool.cpp
--- a/Engine/Plugins/Experimental/UVEditor/Source/UVEditorTools/Private/UVEditorSeamTool.cpp
+++ b/Engine/Plugins/Experimental/UVEditor/Source/UVEditorTools/Private/UVEditorSeamTool.cpp
@@ -477,8 +477,7 @@ void UUVEditorSeamTool::Setup()
 		});
 	}
 
-	GetToolManager()->DisplayMessage(LOCTEXT("OnStartSeamTool", "Click to add seams in 2D or 3D viewport. To finish the seam, either click on the last point added, press enter or create a loop by clicking on the first point."),
-		EToolMessageLevel::UserNotification);
+	UpdateToolMessage();
 }
 
 void UUVEditorSeamTool::ReconstructExistingSeamsVisualization()
@@ -544,10 +543,35 @@ void UUVEditorSeamTool::ReconstructLockedPathVisualization()
 	}
 }
 
+void UUVEditorSeamTool::UpdateToolMessage()
+{
+	switch (State)
+	{
+	case EState::WaitingToStart:
+		GetToolManager()->DisplayMessage(LOCTEXT("StartSeamMessage", "Click to start a seam in the 2D or 3D viewport."),
+			EToolMessageLevel::UserNotification);
+		break;
+	case EState::SeamInProgress:
+		GetToolManager()->DisplayMessage(LOCTEXT("CompleteSeamMessage", "To complete the seam, press Enter or click either on the first or the last point. Press Esc to cancel."),
+			EToolMessageLevel::UserNotification);
+		break;
+	default:
+		ensure(false);
+		GetToolManager()->DisplayMessage(FText(), EToolMessageLevel::UserNotification);
+		break;
+	}
+}
+
 void UUVEditorSeamTool::Shutdown(EToolShutdownType ShutdownType)
 {
 	TRACE_CPUPROFILER_EVENT_SCOPE(UVEditorSeamTool_Shutdown);
 
+	// Apply any pending seam if needed
+	if (ShutdownType != EToolShutdownType::Cancel && LockedPath.Num() > 0)
+	{
+		ApplySeam(LockedPath);
+	}
+
 	UInteractiveTool::Shutdown(ShutdownType);
 
 	EmitChangeAPI = nullptr;
@@ -751,6 +775,7 @@ void UUVEditorSeamTool::ApplyClick()
 			EditSeamTransactionName);
 
 		State = EState::SeamInProgress;
+		UpdateToolMessage();
 		return;
 	}
 
@@ -904,6 +929,7 @@ void UUVEditorSeamTool::ClearLockedPath(bool bEmitChange)
 	LastLockedAppliedVid = IndexConstants::InvalidID;
 
 	State = EState::WaitingToStart;
+	UpdateToolMessage();
 }
 
 int32 UUVEditorSeamTool::Get2DHitVertex(const FRay& WorldRayIn, int32* IndexOf2DSpatialOut)
@@ -1110,6 +1136,7 @@ void UUVEditorSeamTool::EditLockedPath(
 
 	bool bPathNotEmpty = LockedPath.Num() > 0;
 	State = bPathNotEmpty ? EState::SeamInProgress : EState::WaitingToStart;
+	UpdateToolMessage();
 	ClickedMeshIndex = bPathNotEmpty ? MeshIndex : IndexConstants::InvalidID;
 	LastLockedAppliedVid = bPathNotEmpty ? LockedPath.Last() : IndexConstants::InvalidID;
 	SeamStartAppliedVid = bPathNotEmpty ? LockedPath[0] : IndexConstants::InvalidID;
diff --git a/Engine/Plugins/Experimental/UVEditor/Source/UVEditorTools/Public/UVEditorSeamTool.h b/Engine/Plugins/Experimental/UVEditor/Source/UVEditorTools/Public/UVEditorSeamTool.h
--- a/Engine/Plugins/Experimental/UVEditor/Source/UVEditorTools/Public/UVEditorSeamTool.h
+++ b/Engine/Plugins/Experimental/UVEditor/Source/UVEditorTools/Public/UVEditorSeamTool.h
@@ -84,6 +84,8 @@ protected:
 	void ApplySeam(const TArray<int32>& AppliedVidsIn);
 	void ClearLockedPath(bool bEmitChange = true);
 
+	void UpdateToolMessage();
+
 	enum class EState
 	{
 		WaitingToStart,
