UVEditor: Force full update whenever unwrap connectivity changes. This also fixes issues with unset UV tris.

#rnx

diff --git a/Engine/Plugins/Experimental/UVEditor/Source/UVEditorTools/Private/ToolTargets/UVEditorToolMeshInput.cpp b/Engine/Plugins/Experimental/UVEditor/Source/UVEditorTools/Private/ToolTargets/UVEditorToolMeshInput.cpp
--- a/Engine/Plugins/Experimental/UVEditor/Source/UVEditorTools/Private/ToolTargets/UVEditorToolMeshInput.cpp
+++ b/Engine/Plugins/Experimental/UVEditor/Source/UVEditorTools/Private/ToolTargets/UVEditorToolMeshInput.cpp
@@ -19,6 +19,42 @@ using namespace UE::Geometry;
 namespace UVEditorToolMeshInputLocals
 {
 	const TArray<int32> EmptyArray;
+
+	void NotifyUnwrapPreviewDeferredEditCompleted(UUVEditorToolMeshInput* InputObject,
+		const TArray<int32>* ChangedVids, const TArray<int32>* ChangedConnectivityTids, 
+		const TArray<int32>* FastRenderUpdateTids)
+	{
+		EMeshRenderAttributeFlags ChangedAttribs = EMeshRenderAttributeFlags::Positions | EMeshRenderAttributeFlags::VertexUVs;
+
+		// Currently there's no way to do a partial or region update if the connectivity changed in any way.
+		if (ChangedConnectivityTids == nullptr || !ChangedConnectivityTids->IsEmpty())
+		{
+			InputObject->UnwrapPreview->PreviewMesh->NotifyDeferredEditCompleted(UPreviewMesh::ERenderUpdateMode::FullUpdate, ChangedAttribs, true);
+		}
+		else if (FastRenderUpdateTids) 
+		{
+			InputObject->UnwrapPreview->PreviewMesh->NotifyRegionDeferredEditCompleted(*FastRenderUpdateTids, ChangedAttribs);
+		}
+		else
+		{
+			InputObject->UnwrapPreview->PreviewMesh->NotifyDeferredEditCompleted(UPreviewMesh::ERenderUpdateMode::FastUpdate, ChangedAttribs, true);
+		}
+	}
+
+	void NotifyAppliedPreviewDeferredEditCompleted(UUVEditorToolMeshInput* InputObject,
+		const TArray<int32>* ChangedVids, const TArray<int32>* FastRenderUpdateTids)
+	{
+		EMeshRenderAttributeFlags ChangedAttribs = EMeshRenderAttributeFlags::VertexUVs;
+
+		if (FastRenderUpdateTids) 
+		{
+			InputObject->AppliedPreview->PreviewMesh->NotifyRegionDeferredEditCompleted(*FastRenderUpdateTids, ChangedAttribs);
+		}
+		else
+		{
+			InputObject->AppliedPreview->PreviewMesh->NotifyDeferredEditCompleted(UPreviewMesh::ERenderUpdateMode::FastUpdate, ChangedAttribs, false);
+		}
+	}
 }
 
 const TArray<int32>* const UUVEditorToolMeshInput::NONE_CHANGED_ARG = &UVEditorToolMeshInputLocals::EmptyArray;
@@ -101,6 +137,8 @@ void UUVEditorToolMeshInput::Shutdown()
 void UUVEditorToolMeshInput::UpdateUnwrapPreviewOverlayFromPositions(const TArray<int32>* ChangedVids, 
 	const TArray<int32>* ChangedConnectivityTids, const TArray<int32>* FastRenderUpdateTids)
 {
+	using namespace UVEditorToolMeshInputLocals;
+
 	UnwrapPreview->PreviewMesh->DeferredEditMesh([this, ChangedVids, ChangedConnectivityTids](FDynamicMesh3& Mesh)
 	{
 		FDynamicMeshUVOverlay* DestOverlay = Mesh.Attributes()->PrimaryUV();
@@ -108,15 +146,9 @@ void UUVEditorToolMeshInput::UpdateUnwrapPreviewOverlayFromPositions(const TArra
 			ChangedVids, ChangedConnectivityTids);
 	}, false);
 
-	if (FastRenderUpdateTids) {
-		UnwrapPreview->PreviewMesh->NotifyRegionDeferredEditCompleted(*FastRenderUpdateTids,
-			EMeshRenderAttributeFlags::Positions | EMeshRenderAttributeFlags::VertexUVs); // Assuming user changed positions and didn't notify yet
-	} 
-	else
-	{
-		UnwrapPreview->PreviewMesh->NotifyDeferredEditCompleted(UPreviewMesh::ERenderUpdateMode::FastUpdate,
-			EMeshRenderAttributeFlags::Positions | EMeshRenderAttributeFlags::VertexUVs, true); // Assuming user changed positions and didn't notify yet
-	}
+	// We only updated UV's, but this update usually comes after a change to unwrap vert positions, and we
+	// can assume that the user hasn't issued a notification yet. So we do a normal unwrap preview update.
+	NotifyUnwrapPreviewDeferredEditCompleted(this, ChangedVids, ChangedConnectivityTids, FastRenderUpdateTids);
 
 	if (WireframeDisplay)
 	{
@@ -135,6 +167,8 @@ void UUVEditorToolMeshInput::UpdateUnwrapCanonicalOverlayFromPositions(const TAr
 void UUVEditorToolMeshInput::UpdateAppliedPreviewFromUnwrapPreview(const TArray<int32>* ChangedVids, 
 	const TArray<int32>* ChangedConnectivityTids, const TArray<int32>* FastRenderUpdateTids)
 {
+	using namespace UVEditorToolMeshInputLocals;
+
 	const FDynamicMesh3* SourceUnwrapMesh = UnwrapPreview->PreviewMesh->GetMesh();
 
 	// We need to update the overlay in AppliedPreview. Assuming that the overlay in UnwrapPreview is updated, we can
@@ -147,20 +181,14 @@ void UUVEditorToolMeshInput::UpdateAppliedPreviewFromUnwrapPreview(const TArray<
 			UVUnwrapMeshUtil::UpdateOverlayFromOverlay(*SourceOverlay, *DestOverlay, false, ChangedVids, ChangedConnectivityTids);
 	}, false);
 
-	if (FastRenderUpdateTids) {
-		AppliedPreview->PreviewMesh->NotifyRegionDeferredEditCompleted(*FastRenderUpdateTids,
-			EMeshRenderAttributeFlags::VertexUVs); // Assuming user changed positions and didn't notify yet
-	}
-	else
-	{
-		AppliedPreview->PreviewMesh->NotifyDeferredEditCompleted(UPreviewMesh::ERenderUpdateMode::FastUpdate,
-			EMeshRenderAttributeFlags::VertexUVs, false); // Assuming user changed positions and didn't notify yet
-	}
+	NotifyAppliedPreviewDeferredEditCompleted(this, ChangedVids, FastRenderUpdateTids);
 }
 
 void UUVEditorToolMeshInput::UpdateUnwrapPreviewFromAppliedPreview(const TArray<int32>* ChangedElementIDs, 
 	const TArray<int32>* ChangedConnectivityTids, const TArray<int32>* FastRenderUpdateTids)
 {
+	using namespace UVEditorToolMeshInputLocals;
+
 	// Convert the AppliedPreview UV overlay to positions in UnwrapPreview
 	const FDynamicMeshUVOverlay* SourceOverlay = AppliedPreview->PreviewMesh->GetMesh()->Attributes()->GetUVLayer(UVLayerIndex);
 	UnwrapPreview->PreviewMesh->DeferredEditMesh([this, SourceOverlay, ChangedElementIDs, ChangedConnectivityTids](FDynamicMesh3& Mesh)
@@ -172,15 +200,7 @@ void UUVEditorToolMeshInput::UpdateUnwrapPreviewFromAppliedPreview(const TArray<
 		UVUnwrapMeshUtil::UpdateOverlayFromOverlay(*SourceOverlay, *DestOverlay, false, ChangedElementIDs, ChangedConnectivityTids);
 	}, false);
 
-	if (FastRenderUpdateTids) {
-		UnwrapPreview->PreviewMesh->NotifyRegionDeferredEditCompleted(*FastRenderUpdateTids,
-			EMeshRenderAttributeFlags::Positions | EMeshRenderAttributeFlags::VertexUVs); // Assuming user changed positions and didn't notify yet
-	}
-	else
-	{
-		UnwrapPreview->PreviewMesh->NotifyDeferredEditCompleted(UPreviewMesh::ERenderUpdateMode::FastUpdate,
-			EMeshRenderAttributeFlags::Positions | EMeshRenderAttributeFlags::VertexUVs, true); // Assuming user changed positions and didn't notify yet
-	}
+	NotifyUnwrapPreviewDeferredEditCompleted(this, ChangedElementIDs, ChangedConnectivityTids, FastRenderUpdateTids);
 
 	if (WireframeDisplay)
 	{
@@ -207,20 +227,14 @@ void UUVEditorToolMeshInput::UpdateCanonicalFromPreviews(const TArray<int32>* Ch
 
 void UUVEditorToolMeshInput::UpdatePreviewsFromCanonical(const TArray<int32>* ChangedVids, const TArray<int32>* ChangedConnectivityTids, const TArray<int32>* FastRenderUpdateTids)
 {
+	using namespace UVEditorToolMeshInputLocals;
+
 	// Update UnwrapPreview from UnwrapCanonical
 	UnwrapPreview->PreviewMesh->DeferredEditMesh([this, ChangedVids, ChangedConnectivityTids](FDynamicMesh3& Mesh)
 	{
 		UVUnwrapMeshUtil::UpdateUVUnwrapMesh(*UnwrapCanonical, Mesh, ChangedVids, ChangedConnectivityTids);
 	}, false);
-	if (FastRenderUpdateTids) {
-		UnwrapPreview->PreviewMesh->NotifyRegionDeferredEditCompleted(*FastRenderUpdateTids,
-			EMeshRenderAttributeFlags::Positions | EMeshRenderAttributeFlags::VertexUVs);
-	}
-	else
-	{
-		UnwrapPreview->PreviewMesh->NotifyDeferredEditCompleted(UPreviewMesh::ERenderUpdateMode::FastUpdate,
-			EMeshRenderAttributeFlags::Positions | EMeshRenderAttributeFlags::VertexUVs, true);
-	}
+	NotifyUnwrapPreviewDeferredEditCompleted(this, ChangedVids, ChangedConnectivityTids, FastRenderUpdateTids);
 
 	if (WireframeDisplay)
 	{
@@ -234,15 +248,7 @@ void UUVEditorToolMeshInput::UpdatePreviewsFromCanonical(const TArray<int32>* Ch
 		FDynamicMeshUVOverlay* DestOverlay = Mesh.Attributes()->GetUVLayer(UVLayerIndex);
 		UVUnwrapMeshUtil::UpdateOverlayFromOverlay(*SourceOverlay, *DestOverlay, true, ChangedVids, ChangedConnectivityTids);
 	}, false);
-	if (FastRenderUpdateTids) {
-		AppliedPreview->PreviewMesh->NotifyRegionDeferredEditCompleted(*FastRenderUpdateTids,
-			EMeshRenderAttributeFlags::VertexUVs);
-	}
-	else
-	{
-		AppliedPreview->PreviewMesh->NotifyDeferredEditCompleted(UPreviewMesh::ERenderUpdateMode::FastUpdate,
-			EMeshRenderAttributeFlags::VertexUVs, true);
-	}
+	NotifyAppliedPreviewDeferredEditCompleted(this, ChangedVids, FastRenderUpdateTids);
 }
 
 void UUVEditorToolMeshInput::UpdateAllFromUnwrapPreview(const TArray<int32>* ChangedVids, 
