UVEditor: Change Seam tool to use FUVEditorUXSettings.

#rnx

diff --git a/Engine/Plugins/Experimental/UVEditor/Source/UVEditorTools/Private/UVEditorSeamTool.cpp b/Engine/Plugins/Experimental/UVEditor/Source/UVEditorTools/Private/UVEditorSeamTool.cpp
--- a/Engine/Plugins/Experimental/UVEditor/Source/UVEditorTools/Private/UVEditorSeamTool.cpp
+++ b/Engine/Plugins/Experimental/UVEditor/Source/UVEditorTools/Private/UVEditorSeamTool.cpp
@@ -15,6 +15,7 @@
 #include "Parameterization/MeshDijkstra.h"
 #include "ToolTargets/UVEditorToolMeshInput.h"
 #include "ToolSceneQueriesUtil.h"
+#include "UVEditorUXSettings.h"
 #include "UVToolContextObjects.h"
 
 #define LOCTEXT_NAMESPACE "UUVEditorSeamTool"
@@ -36,22 +37,6 @@ namespace UVEditorSeamToolLocals
 	
 	const FString& ExistingSeamsID(TEXT("SeamLineSet"));
 
-	const FColor ExistingSeamColor = FColor::Green;
-	const FColor LockedColor = FColor::Green;
-	const FColor CompletionColor = FColor::Orange;
-	const FColor HoverColor = FColor::Green;
-
-	float ExistingSeamThickness = 2;
-	float ExistingSeamDepthBias = 0.1;
-
-	float LockedThickness = 3;
-	float LockedDepthBias = 1.0;
-
-	float HoverThickness = 3;
-	float HoverDepthBias = 1.1;
-
-	const float PointSize = 7;
-
 	void GetVidPath(FDynamicMesh3* Mesh, const TArray<int32>& StartVids, int32 EndVid, TArray<int32>& VidPathOut)
 	{
 		VidPathOut.Reset();
@@ -89,11 +74,11 @@ namespace UVEditorSeamToolLocals
 
 		FTransform AppliedTransform = InputObject->AppliedPreview->PreviewMesh->GetTransform();
 		AppliedPointSet->AddPoint(FRenderablePoint(
-			InputObject->AppliedCanonical->GetVertex(AppliedVid), Color, PointSize, DepthBias));
+			InputObject->AppliedCanonical->GetVertex(AppliedVid), Color, FUVEditorUXSettings::ToolPointSize, DepthBias));
 		for (int32 UnwrapVid : UnwrapVids)
 		{
 			UnwrapPointSet->AddPoint(FRenderablePoint(
-				InputObject->UnwrapCanonical->GetVertex(UnwrapVid), Color, PointSize, DepthBias));
+				InputObject->UnwrapCanonical->GetVertex(UnwrapVid), Color, FUVEditorUXSettings::ToolPointSize, DepthBias));
 		}
 	}
 
@@ -500,7 +485,9 @@ void UUVEditorSeamTool::ReconstructExistingSeamsVisualization()
 				FIndex2i EdgeVids = Mesh->GetEdgeV(Eid);
 				SeamLines->AddLine((FVector)Transform.TransformPosition(Mesh->GetVertex(EdgeVids.A)),
 					(FVector)Transform.TransformPosition(Mesh->GetVertex(EdgeVids.B)),
-					ExistingSeamColor, ExistingSeamThickness, ExistingSeamDepthBias);
+					FUVEditorUXSettings::LivePreviewExistingSeamColor,
+					FUVEditorUXSettings::LivePreviewExistingSeamThickness,
+					FUVEditorUXSettings::LivePreviewExistingSeamDepthBias);
 			}
 		}
 	}
@@ -531,15 +518,19 @@ void UUVEditorSeamTool::ReconstructLockedPathVisualization()
 		UUVEditorToolMeshInput* ClickedTarget = Targets[ClickedMeshIndex];
 
 		AddDisplayedPath(ClickedTarget, UnwrapLineSet, AppliedLineSet, LockedPath,
-			false, LockedColor, LockedThickness, LockedDepthBias);
+			false, FUVEditorUXSettings::ToolLockedPathColor,
+			FUVEditorUXSettings::ToolLockedPathThickness,
+			FUVEditorUXSettings::ToolLockedPathDepthBias);
 
 		TArray<int32> UnwrapVids;
 		ClickedTarget->AppliedVidToUnwrapVids(LastLockedAppliedVid, UnwrapVids);
-		AddDisplayedPoints(ClickedTarget, UnwrapPointSet, AppliedPointSet, LastLockedAppliedVid, UnwrapVids, LockedColor, LockedDepthBias);
+		AddDisplayedPoints(ClickedTarget, UnwrapPointSet, AppliedPointSet, LastLockedAppliedVid, UnwrapVids, 
+			FUVEditorUXSettings::ToolLockedPathColor, FUVEditorUXSettings::ToolLockedPathDepthBias);
 
 		UnwrapVids.Reset();
 		ClickedTarget->AppliedVidToUnwrapVids(SeamStartAppliedVid, UnwrapVids);
-		AddDisplayedPoints(ClickedTarget, UnwrapStartPoints, AppliedStartPoints, SeamStartAppliedVid, UnwrapVids, LockedColor, LockedDepthBias);
+		AddDisplayedPoints(ClickedTarget, UnwrapStartPoints, AppliedStartPoints, SeamStartAppliedVid, UnwrapVids,
+			FUVEditorUXSettings::ToolLockedPathColor, FUVEditorUXSettings::ToolLockedPathDepthBias);
 	}
 }
 
@@ -654,7 +645,7 @@ void UUVEditorSeamTool::UpdateHover()
 	{
 		// Just draw the point and finish
 		AddDisplayedPoints(Targets[HoverMeshIndex], UnwrapPointSet, AppliedPointSet,
-			AppliedVid, UnwrapVids, HoverColor, HoverDepthBias);
+			AppliedVid, UnwrapVids, FUVEditorUXSettings::ToolExtendPathColor, FUVEditorUXSettings::ToolExtendPathDepthBias);
 		return;
 	}
 
@@ -681,7 +672,8 @@ void UUVEditorSeamTool::UpdateHover()
 		ULineSetComponent* UnwrapLineSet = UnwrapGeometry->FindLineSet(HoverLineSetID);
 		ULineSetComponent* AppliedLineSet = LivePreviewGeometry->FindLineSet(HoverLineSetID);
 		AddDisplayedPath(Targets[HoverMeshIndex], UnwrapLineSet, AppliedLineSet,
-			NewPathVids, bHoverVidIsFromUnwrap, HoverColor, HoverThickness, HoverDepthBias);
+			NewPathVids, bHoverVidIsFromUnwrap, FUVEditorUXSettings::ToolExtendPathColor,
+			FUVEditorUXSettings::ToolExtendPathThickness, FUVEditorUXSettings::ToolExtendPathDepthBias);
 	}
 
 	// See if we would apply the seam on click
@@ -696,7 +688,7 @@ void UUVEditorSeamTool::UpdateHover()
 		// there because that turns out to be confusing, esp when we have overlapping islands and
 		// the point is on another island.
 		AddDisplayedPoints(Targets[HoverMeshIndex], UnwrapPointSet, AppliedPointSet,
-			AppliedVid, UnwrapVids, HoverColor, HoverDepthBias);
+			AppliedVid, UnwrapVids, FUVEditorUXSettings::ToolExtendPathColor, FUVEditorUXSettings::ToolExtendPathDepthBias);
 	}
 }
 
@@ -728,14 +720,16 @@ void UUVEditorSeamTool::ResetPreviewColors()
 {
 	using namespace UVEditorSeamToolLocals;
 
-	FColor Color = bCompletionColorOverride ? CompletionColor : HoverColor;
+	FColor Color = bCompletionColorOverride ? FUVEditorUXSettings::ToolCompletionPathColor 
+		: FUVEditorUXSettings::ToolExtendPathColor;
 
 	UnwrapGeometry->FindPointSet(HoverPointSetID)->SetAllPointsColor(Color);
 	LivePreviewGeometry->FindPointSet(HoverPointSetID)->SetAllPointsColor(Color);
 	UnwrapGeometry->FindLineSet(HoverLineSetID)->SetAllLinesColor(Color);
 	LivePreviewGeometry->FindLineSet(HoverLineSetID)->SetAllLinesColor(Color);
 
-	Color = bCompletionColorOverride ? CompletionColor : LockedColor;
+	Color = bCompletionColorOverride ? FUVEditorUXSettings::ToolCompletionPathColor 
+		: FUVEditorUXSettings::ToolLockedPathColor;
 
 	UnwrapGeometry->FindPointSet(LockedPointSetID)->SetAllPointsColor(Color);
 	LivePreviewGeometry->FindPointSet(LockedPointSetID)->SetAllPointsColor(Color);
@@ -767,7 +761,8 @@ void UUVEditorSeamTool::ApplyClick()
 
 		UPointSetComponent* UnwrapPointSet = UnwrapGeometry->FindPointSet(StartPointSetID);
 		UPointSetComponent* AppliedPointSet = LivePreviewGeometry->FindPointSet(StartPointSetID);
-		AddDisplayedPoints(ClickedTarget, UnwrapPointSet, AppliedPointSet, AppliedVid, UnwrapVids, LockedColor, LockedDepthBias);
+		AddDisplayedPoints(ClickedTarget, UnwrapPointSet, AppliedPointSet, AppliedVid, UnwrapVids, 
+			FUVEditorUXSettings::ToolLockedPathColor, FUVEditorUXSettings::ToolLockedPathDepthBias);
 
 		LockedPath.Add(AppliedVid);
 		EmitChangeAPI->EmitToolDependentChange(this, 
@@ -835,10 +830,13 @@ void UUVEditorSeamTool::ApplyClick()
 	ULineSetComponent* UnwrapLineSet = UnwrapGeometry->FindLineSet(LockedLineSetID);
 	ULineSetComponent* AppliedLineSet = LivePreviewGeometry->FindLineSet(LockedLineSetID);
 	AddDisplayedPath(ClickedTarget, UnwrapLineSet, AppliedLineSet, VidPath,
-		bClickWasInUnwrap, LockedColor, LockedThickness, LockedDepthBias);
+		bClickWasInUnwrap, FUVEditorUXSettings::ToolLockedPathColor,
+		FUVEditorUXSettings::ToolLockedPathThickness,
+		FUVEditorUXSettings::ToolLockedPathDepthBias);
 	UnwrapPointSet->Clear();
 	AppliedPointSet->Clear();
-	AddDisplayedPoints(ClickedTarget, UnwrapPointSet, AppliedPointSet, AppliedVid, UnwrapVids, LockedColor, LockedDepthBias);
+	AddDisplayedPoints(ClickedTarget, UnwrapPointSet, AppliedPointSet, AppliedVid, UnwrapVids, 
+		FUVEditorUXSettings::ToolLockedPathColor, FUVEditorUXSettings::ToolLockedPathDepthBias);
 	
 	EmitChangeAPI->EmitToolDependentChange(this, 
 		MakeUnique<FPathChange>(*AppliedVidPath, false, ClickedMeshIndex, true),
diff --git a/Engine/Plugins/Experimental/UVEditor/Source/UVEditorTools/Private/UVEditorUXSettings.cpp b/Engine/Plugins/Experimental/UVEditor/Source/UVEditorTools/Private/UVEditorUXSettings.cpp
--- a/Engine/Plugins/Experimental/UVEditor/Source/UVEditorTools/Private/UVEditorUXSettings.cpp
+++ b/Engine/Plugins/Experimental/UVEditor/Source/UVEditorTools/Private/UVEditorUXSettings.cpp
@@ -3,12 +3,14 @@
 #include "UVEditorUXSettings.h"
 
 const float FUVEditorUXSettings::CameraFarPlaneWorldZ(-10.0);
+const float FUVEditorUXSettings::CameraNearPlaneProportionZ(0.8); // Top layer, equivalent to depth bias 80
 
 // 2D Viewport Depth Offsets (Organized by "layers" from the camera's perspective, descending order
 // Note: While these are floating point values, they represent percentages and should be separated
 // by at least integer amounts, as they serve double duty in certain cases for translucent primitive
 // sorting order.
-const float FUVEditorUXSettings::CameraNearPlaneProportionZ(8); // Top layer
+const float FUVEditorUXSettings::ToolLockedPathDepthBias(6.0);
+const float FUVEditorUXSettings::ToolExtendPathDepthBias(6.0);
 const float FUVEditorUXSettings::SewLineDepthOffset(5.0f);
 const float FUVEditorUXSettings::SelectionHoverWireframeDepthBias(6);
 const float FUVEditorUXSettings::SelectionHoverTriangleDepthBias(5);
@@ -17,6 +19,8 @@ const float FUVEditorUXSettings::SelectionTriangleDepthBias(3.0);
 const float FUVEditorUXSettings::WireframeDepthOffset(2.0);
 const float FUVEditorUXSettings::UnwrapTriangleDepthOffset(1.0);
 
+const float FUVEditorUXSettings::LivePreviewExistingSeamDepthBias(1.0);
+
 // Note: that this offset can only be applied when we use our own background material
 // for a user-supplied texture, and we can't use it for a user-provided material.
 // So for consistency this should stay at zero.
@@ -46,6 +50,13 @@ const FColor FUVEditorUXSettings::SelectionHoverTriangleFillColor(FColor::FromHe
 const FColor FUVEditorUXSettings::SelectionHoverTriangleWireframeColor(FColor::FromHex("#0E86FF"));
 const FColor FUVEditorUXSettings::SewSideLeftColor(FColor::Red);
 const FColor FUVEditorUXSettings::SewSideRightColor(FColor::Green);
+
+const FColor FUVEditorUXSettings::ToolLockedPathColor(FColor::Green);
+const FColor FUVEditorUXSettings::ToolExtendPathColor(FColor::Green);
+const FColor FUVEditorUXSettings::ToolCompletionPathColor(FColor::Orange);
+
+const FColor FUVEditorUXSettings::LivePreviewExistingSeamColor(FColor::Green);
+
 const FColor FUVEditorUXSettings::XAxisColor(FColor::Red);
 const FColor FUVEditorUXSettings::YAxisColor(FColor::Green);
 const FColor FUVEditorUXSettings::GridMajorColor(FColor::FromHex("#888888"));
@@ -54,12 +65,17 @@ const FColor FUVEditorUXSettings::GridMinorColor(FColor::FromHex("#777777"));
 // Thicknesses
 const float FUVEditorUXSettings::LivePreviewHighlightThickness(2.0);
 const float FUVEditorUXSettings::LivePreviewHighlightPointSize(4);
+const float FUVEditorUXSettings::LivePreviewExistingSeamThickness(2.0);
 const float FUVEditorUXSettings::SelectionLineThickness(1.5);
+const float FUVEditorUXSettings::ToolLockedPathThickness(3.0f);
+const float FUVEditorUXSettings::ToolExtendPathThickness(3.0f);
 const float FUVEditorUXSettings::SelectionPointThickness(6);
 const float FUVEditorUXSettings::SewLineHighlightThickness(3.0f);
 const float FUVEditorUXSettings::AxisThickness(2.0);
 const float FUVEditorUXSettings::GridMajorThickness(1.0);
 
+const float FUVEditorUXSettings::ToolPointSize(6);
+
 // Grid
 const int32 FUVEditorUXSettings::GridSubdivisionsPerLevel(4);
 const int32 FUVEditorUXSettings::GridLevels(3);
diff --git a/Engine/Plugins/Experimental/UVEditor/Source/UVEditorTools/Public/UVEditorUXSettings.h b/Engine/Plugins/Experimental/UVEditor/Source/UVEditorTools/Public/UVEditorUXSettings.h
--- a/Engine/Plugins/Experimental/UVEditor/Source/UVEditorTools/Public/UVEditorUXSettings.h
+++ b/Engine/Plugins/Experimental/UVEditor/Source/UVEditorTools/Public/UVEditorUXSettings.h
@@ -56,6 +56,22 @@ public:
 	static const float SelectionHoverWireframeDepthBias;
 	static const float SelectionHoverTriangleDepthBias;
 
+	static const FColor LivePreviewExistingSeamColor;
+	static const float LivePreviewExistingSeamThickness;
+	static const float LivePreviewExistingSeamDepthBias;
+
+	// These are currently used by the seam tool but can be generally used by
+	// tools for displaying paths.
+	static const FColor ToolLockedPathColor;
+	static const float ToolLockedPathThickness;
+	static const float ToolLockedPathDepthBias;
+	static const FColor ToolExtendPathColor;
+	static const float ToolExtendPathThickness;
+	static const float ToolExtendPathDepthBias;
+	static const FColor ToolCompletionPathColor;
+
+	static const float ToolPointSize;
+
 	// Sew Action styling
 	static const float SewLineHighlightThickness;
 	static const float SewLineDepthOffset;
