UVEditor: Cleans up issues with mode layout, preventing ghost tabs from appearing and moving UV Editor specific tabs to their own menu category.

This addresses issues with a previous CL review: https://p4-swarm.epicgames.net/reviews/18157315

#rnx

diff --git a/Engine/Plugins/Experimental/UVEditor/Source/UVEditor/Private/UVEditorModeUILayer.cpp b/Engine/Plugins/Experimental/UVEditor/Source/UVEditor/Private/UVEditorModeUILayer.cpp
--- a/Engine/Plugins/Experimental/UVEditor/Source/UVEditor/Private/UVEditorModeUILayer.cpp
+++ b/Engine/Plugins/Experimental/UVEditor/Source/UVEditor/Private/UVEditorModeUILayer.cpp
@@ -19,8 +19,8 @@ void UUVEditorUISubsystem::Deinitialize()
 }
 void UUVEditorUISubsystem::RegisterLayoutExtensions(FLayoutExtender& Extender)
 {
-	FTabManager::FTab NewTab(FTabId(UAssetEditorUISubsystem::TopLeftTabID, ETabIdFlags::SaveLayout), ETabState::ClosedTab);
-	Extender.ExtendStack("ToolbarArea", ELayoutExtensionPosition::After, NewTab);
+	FTabManager::FTab NewTab(FTabId(UAssetEditorUISubsystem::TopLeftTabID), ETabState::ClosedTab);
+	Extender.ExtendStack("EditorSidePanelArea", ELayoutExtensionPosition::After, NewTab);
 }
 
 FUVEditorModeUILayer::FUVEditorModeUILayer(const IToolkitHost* InToolkitHost) :
@@ -50,6 +50,11 @@ void FUVEditorModeUILayer::OnToolkitHostingFinished(const TSharedRef<IToolkit>&
 
 TSharedPtr<FWorkspaceItem> FUVEditorModeUILayer::GetModeMenuCategory() const
 {
-	const IWorkspaceMenuStructure& MenuStructure = WorkspaceMenu::GetMenuStructure();
-	return MenuStructure.GetLevelEditorModesCategory();
-}
\ No newline at end of file
+	check(UVEditorMenuCategory);
+	return UVEditorMenuCategory;
+}
+
+void FUVEditorModeUILayer::SetModeMenuCategory(TSharedPtr<FWorkspaceItem> MenuCategoryIn)
+{
+	UVEditorMenuCategory = MenuCategoryIn;
+}
diff --git a/Engine/Plugins/Experimental/UVEditor/Source/UVEditor/Private/UVEditorToolkit.cpp b/Engine/Plugins/Experimental/UVEditor/Source/UVEditor/Private/UVEditorToolkit.cpp
--- a/Engine/Plugins/Experimental/UVEditor/Source/UVEditor/Private/UVEditorToolkit.cpp
+++ b/Engine/Plugins/Experimental/UVEditor/Source/UVEditor/Private/UVEditorToolkit.cpp
@@ -47,8 +47,11 @@ FUVEditorToolkit::FUVEditorToolkit(UAssetEditor* InOwningAssetEditor)
 	// are two viewports on the right.
 	// We define explicit ExtensionIds on the stacks to reference them later when the
     // UILayer provides layout extensions. 
-
-	StandaloneDefaultLayout = FTabManager::NewLayout(FName("UVEditorLayout1"))
+	//
+	// Note: Changes to the layout should include a increment to the layout's ID, i.e.
+	// UVEditorLayout[X] -> UVEditorLayout[X+1]. Otherwise, layouts may be messed up
+	// without a full reset to layout defaults inside the editor.
+	StandaloneDefaultLayout = FTabManager::NewLayout(FName("UVEditorLayout2"))
 		->AddArea
 		(
 			FTabManager::NewPrimaryArea()->SetOrientation(Orient_Vertical)
@@ -59,7 +62,7 @@ FUVEditorToolkit::FUVEditorToolkit(UAssetEditor* InOwningAssetEditor)
 				(
 					FTabManager::NewStack()
 					->SetSizeCoefficient(0.2f)					
-					->SetExtensionId("ToolbarArea")
+					->SetExtensionId("EditorSidePanelArea")
 					->SetHideTabWell(true)
 				)				
 				->Split
@@ -183,19 +186,21 @@ void FUVEditorToolkit::RegisterTabSpawners(const TSharedRef<FTabManager>& InTabM
 	// anything except tabs that we don't want.
 	FAssetEditorToolkit::RegisterTabSpawners(InTabManager);
 
+	UVEditorMenuCategory = InTabManager->AddLocalWorkspaceMenuCategory(LOCTEXT("WorkspaceMenu_UVEditor", "UV Editor"));
+
 	// Here we set up the tabs we referenced in StandaloneDefaultLayout (in the constructor).
 	// We don't deal with the toolbar palette here, since this is handled by existing
     // infrastructure in FModeToolkit. We only setup spawners for our custom tabs, namely
     // the 2D and 3D viewports.
 	InTabManager->RegisterTabSpawner(ViewportTabID, FOnSpawnTab::CreateSP(this, &FUVEditorToolkit::SpawnTab_Viewport))
 		.SetDisplayName(LOCTEXT("2DViewportTabLabel", "2D Viewport"))
-		.SetGroup(AssetEditorTabsCategory.ToSharedRef())
+		.SetGroup(UVEditorMenuCategory.ToSharedRef())
 		.SetIcon(FSlateIcon(FEditorStyle::GetStyleSetName(), "LevelEditor.Tabs.Viewports"));
 
 	InTabManager->RegisterTabSpawner(LivePreviewTabID, FOnSpawnTab::CreateSP(this, 
 		&FUVEditorToolkit::SpawnTab_LivePreview))
 		.SetDisplayName(LOCTEXT("3DViewportTabLabel", "3D Viewport"))
-		.SetGroup(AssetEditorTabsCategory.ToSharedRef())
+		.SetGroup(UVEditorMenuCategory.ToSharedRef())
 		.SetIcon(FSlateIcon(FEditorStyle::GetStyleSetName(), "LevelEditor.Tabs.Viewports"));
 }
 
@@ -366,7 +371,7 @@ void FUVEditorToolkit::PostInitAssetEditor()
 	TSharedPtr<class IToolkitHost> PinnedToolkitHost = ToolkitHost.Pin();
 	check(PinnedToolkitHost.IsValid());
 	ModeUILayer = MakeShareable(new FUVEditorModeUILayer(PinnedToolkitHost.Get()));
-
+	ModeUILayer->SetModeMenuCategory( UVEditorMenuCategory );
 	// Currently, aside from setting up all the UI elements, the toolkit also kicks off the UV
 	// editor mode, which is the mode that the editor always works in (things are packaged into
 	// a mode so that they can be moved to another asset editor if necessary).
diff --git a/Engine/Plugins/Experimental/UVEditor/Source/UVEditor/Public/UVEditorModeUILayer.h b/Engine/Plugins/Experimental/UVEditor/Source/UVEditor/Public/UVEditorModeUILayer.h
--- a/Engine/Plugins/Experimental/UVEditor/Source/UVEditor/Public/UVEditorModeUILayer.h
+++ b/Engine/Plugins/Experimental/UVEditor/Source/UVEditor/Public/UVEditorModeUILayer.h
@@ -24,5 +24,10 @@ public:
 	void OnToolkitHostingStarted(const TSharedRef<IToolkit>& Toolkit) override;
 	void OnToolkitHostingFinished(const TSharedRef<IToolkit>& Toolkit) override;
 
+	void SetModeMenuCategory(TSharedPtr<FWorkspaceItem> MenuCategoryIn);
 	TSharedPtr<FWorkspaceItem> GetModeMenuCategory() const override;
+
+protected:
+	TSharedPtr<FWorkspaceItem> UVEditorMenuCategory;
+
 };
diff --git a/Engine/Plugins/Experimental/UVEditor/Source/UVEditor/Public/UVEditorToolkit.h b/Engine/Plugins/Experimental/UVEditor/Source/UVEditor/Public/UVEditorToolkit.h
--- a/Engine/Plugins/Experimental/UVEditor/Source/UVEditor/Public/UVEditorToolkit.h
+++ b/Engine/Plugins/Experimental/UVEditor/Source/UVEditor/Public/UVEditorToolkit.h
@@ -90,4 +90,5 @@ protected:
 	UUVToolViewportButtonsAPI* ViewportButtonsAPI = nullptr;
 
 	TSharedPtr<FUVEditorModeUILayer> ModeUILayer;
+	TSharedPtr<FWorkspaceItem> UVEditorMenuCategory;
 };
\ No newline at end of file
