UVEditor: Change tab icon color to be always white.

#rnx

diff --git a/Engine/Plugins/Experimental/UVEditor/Source/UVEditor/Private/UVEditorToolkit.cpp b/Engine/Plugins/Experimental/UVEditor/Source/UVEditor/Private/UVEditorToolkit.cpp
--- a/Engine/Plugins/Experimental/UVEditor/Source/UVEditor/Private/UVEditorToolkit.cpp
+++ b/Engine/Plugins/Experimental/UVEditor/Source/UVEditor/Private/UVEditorToolkit.cpp
@@ -143,10 +143,10 @@ FText FUVEditorToolkit::GetToolkitName() const
 	const TArray<UObject*>* Objects = GetObjectsCurrentlyBeingEdited();
 	if (Objects->Num() == 1)
 	{
-		return FText::Format(LOCTEXT("UVEditorTabNameWithObject", "UVs: {0}"), 
+		return FText::Format(LOCTEXT("UVEditorTabNameWithObject", "UV: {0}"), 
 			GetLabelForObject((*Objects)[0]));
 	}
-	return LOCTEXT("UVEditorMultipleTabName", "UVs: Multiple");
+	return LOCTEXT("UVEditorMultipleTabName", "UV: Multiple");
 }
 
 // This gets used multiple places, most notably in GetToolMenuAppName, which gets
@@ -520,6 +520,11 @@ const FSlateBrush* FUVEditorToolkit::GetDefaultTabIcon() const
 	return FUVEditorStyle::Get().GetBrush("UVEditor.OpenUVEditor");
 }
 
+FLinearColor FUVEditorToolkit::GetDefaultTabColor() const
+{
+	return FLinearColor::White;
+}
+
 void FUVEditorToolkit::OnToolkitHostingStarted(const TSharedRef<IToolkit>& Toolkit)
 {
 	ModeUILayer->OnToolkitHostingStarted(Toolkit);
diff --git a/Engine/Plugins/Experimental/UVEditor/Source/UVEditor/Public/UVEditorToolkit.h b/Engine/Plugins/Experimental/UVEditor/Source/UVEditor/Public/UVEditorToolkit.h
--- a/Engine/Plugins/Experimental/UVEditor/Source/UVEditor/Public/UVEditorToolkit.h
+++ b/Engine/Plugins/Experimental/UVEditor/Source/UVEditor/Public/UVEditorToolkit.h
@@ -70,6 +70,7 @@ protected:
 	// FAssetEditorToolkit
 	virtual void PostInitAssetEditor() override;
 	virtual const FSlateBrush* GetDefaultTabIcon() const override;
+	virtual FLinearColor GetDefaultTabColor() const override;
 
 	/** Scene in which the 2D unwrapped uv meshes live. */
 	TUniquePtr<FPreviewScene> UnwrapScene;
