UV Editor: Add some analytics to UV Editor tools

#rnx

diff --git a/Engine/Plugins/Experimental/UVEditor/Source/UVEditor/Private/UVEditorMode.cpp b/Engine/Plugins/Experimental/UVEditor/Source/UVEditor/Private/UVEditorMode.cpp
--- a/Engine/Plugins/Experimental/UVEditor/Source/UVEditor/Private/UVEditorMode.cpp
+++ b/Engine/Plugins/Experimental/UVEditor/Source/UVEditor/Private/UVEditorMode.cpp
@@ -35,6 +35,7 @@
 #include "UVEditorSubsystem.h"
 #include "UVToolContextObjects.h"
 #include "UVEditorBackgroundPreview.h"
+#include "UVEditorToolAnalyticsUtils.h"
 #include "Editor.h"
 #include "UVEditorUXSettings.h"
 
@@ -170,7 +171,7 @@ void UUVEditorMode::Enter()
 		TArray<FAnalyticsEventAttribute> Attributes;
 		Attributes.Add(FAnalyticsEventAttribute(TEXT("Timestamp"), AnalyticsLastStartTimestamp.ToString()));
 
-		FEngineAnalytics::GetProvider().RecordEvent(TEXT("Editor.Usage.UVEditor.Enter"), Attributes);
+		FEngineAnalytics::GetProvider().RecordEvent(UVEditorAnalytics::UVEditorAnalyticsEventName(TEXT("Enter")), Attributes);
 	}
 
 	bIsActive = true;
@@ -337,7 +338,7 @@ void UUVEditorMode::Exit()
 		Attributes.Add(FAnalyticsEventAttribute(TEXT("Timestamp"), Now.ToString()));
 		Attributes.Add(FAnalyticsEventAttribute(TEXT("Duration.Seconds"), static_cast<float>(ModeUsageDuration.GetTotalSeconds())));
 
-		FEngineAnalytics::GetProvider().RecordEvent(TEXT("Editor.Usage.UVEditor.Exit"), Attributes);
+		FEngineAnalytics::GetProvider().RecordEvent(UVEditorAnalytics::UVEditorAnalyticsEventName(TEXT("Exit")), Attributes);
 	}
 
 	// ToolsContext->EndTool only shuts the tool on the next tick, and ToolsContext->DeactivateActiveTool is
diff --git a/Engine/Plugins/Experimental/UVEditor/Source/UVEditorTools/Private/UVEditorLayerEditTool.cpp b/Engine/Plugins/Experimental/UVEditor/Source/UVEditorTools/Private/UVEditorLayerEditTool.cpp
--- a/Engine/Plugins/Experimental/UVEditor/Source/UVEditorTools/Private/UVEditorLayerEditTool.cpp
+++ b/Engine/Plugins/Experimental/UVEditor/Source/UVEditorTools/Private/UVEditorLayerEditTool.cpp
@@ -14,6 +14,8 @@
 #include "UVToolContextObjects.h"
 #include "ModelingToolTargetUtil.h"
 #include "Components.h"
+#include "UVEditorToolAnalyticsUtils.h"
+#include "EngineAnalytics.h"
 
 using namespace UE::Geometry;
 
@@ -462,6 +464,8 @@ void UUVEditorChannelEditToolActionPropertySet::PostAction(EChannelEditToolActio
 void UUVEditorChannelEditTool::Setup()
 {
 	check(Targets.Num() > 0);
+	
+	ToolStartTimeAnalytics = FDateTime::UtcNow();
 
 	UInteractiveTool::Setup();
 
@@ -509,6 +513,9 @@ void UUVEditorChannelEditTool::Setup()
 			UpdateChannelSelectionProperties(i);
 		});
 	}
+
+	// Analytics
+	InputTargetAnalytics = UVEditorAnalytics::CollectTargetAnalytics(Targets);
 }
 
 void UUVEditorChannelEditTool::Shutdown(EToolShutdownType ShutdownType)
@@ -519,6 +526,9 @@ void UUVEditorChannelEditTool::Shutdown(EToolShutdownType ShutdownType)
 	}
 
 	Targets.Empty();
+
+	// Analytics
+	RecordAnalytics();
 }
 
 void UUVEditorChannelEditTool::UpdateChannelSelectionProperties(int32 ChangingAsset)
@@ -617,6 +627,12 @@ void UUVEditorChannelEditTool::AddChannel()
 
 	if (NewChannelIndex != -1)
 	{
+		// Analytics
+		FActionHistoryItem Item;
+		Item.ActionType = EChannelEditToolAction::Add;
+		Item.FirstOperandIndex = NewChannelIndex;
+		AnalyticsActionHistory.Add(Item);
+
 		Targets[ActiveAsset]->AppliedPreview->PreviewMesh->UpdatePreview(Targets[ActiveAsset]->AppliedCanonical.Get());
 
 		EmitChangeAPI->BeginUndoTransaction(UUVEditorChannelEditLocals::UVChannelAddTransactionName);
@@ -639,6 +655,13 @@ void UUVEditorChannelEditTool::AddChannel()
 
 void UUVEditorChannelEditTool::CopyChannel()
 {
+	// Analytics
+	FActionHistoryItem Item;
+	Item.ActionType = EChannelEditToolAction::Copy;
+	Item.FirstOperandIndex = ReferenceChannel;
+	Item.SecondOperandIndex = ActiveChannel;
+	AnalyticsActionHistory.Add(Item);
+
 	EmitChangeAPI->EmitToolIndependentChange(this,
 		MakeUnique<UUVEditorChannelEditLocals::FInputObjectUVChannelCopy>(Targets[ActiveAsset], ReferenceChannel, ActiveChannel, 
 			*Targets[ActiveAsset]->AppliedCanonical->Attributes()->GetUVLayer(ActiveChannel)),
@@ -657,6 +680,13 @@ void UUVEditorChannelEditTool::DeleteChannel()
 
 	int32 TotalLayerCount = Targets[ActiveAsset]->AppliedCanonical->Attributes()->NumUVLayers();
 	bool bClearInstead = TotalLayerCount == 1;
+	
+	// Analytics
+	FActionHistoryItem Item;
+	Item.ActionType = EChannelEditToolAction::Delete;
+	Item.FirstOperandIndex = ActiveChannel;
+	Item.bDeleteActionWasActuallyClear = bClearInstead;
+	AnalyticsActionHistory.Add(Item);
 
 	EmitChangeAPI->EmitToolIndependentChange(this,
 		MakeUnique<UUVEditorChannelEditLocals::FInputObjectUVChannelDelete>(Targets[ActiveAsset], ActiveChannel, 
@@ -669,4 +699,54 @@ void UUVEditorChannelEditTool::DeleteChannel()
 	SourceChannelProperties->TargetChannel = SourceChannelProperties->GetUVChannelNames()[Targets[ActiveAsset]->UVLayerIndex];
 }
 
+void UUVEditorChannelEditTool::RecordAnalytics()
+{
+	using namespace UVEditorAnalytics;
+	
+	if (!FEngineAnalytics::IsAvailable())
+	{
+		return;
+	}
+	
+	TArray<FAnalyticsEventAttribute> Attributes;
+	Attributes.Add(FAnalyticsEventAttribute(TEXT("Timestamp"), FDateTime::UtcNow().ToString()));
+	
+	// Tool inputs
+	InputTargetAnalytics.AppendToAttributes(Attributes, "Input");
+
+	// Tool stats
+	auto MaybeAppendStatsToAttributes = [this, &Attributes](EChannelEditToolAction ActionType)
+	{
+		// For now we only care about the number of times actions were issued
+		int32 NumActions = 0;
+		
+		for (const FActionHistoryItem& Item : AnalyticsActionHistory)
+		{
+			if (Item.ActionType == ActionType) {
+				NumActions += 1;
+			}
+		}
+		
+		if (NumActions > 0)
+		{
+			const FString ActionName = StaticEnum<EChannelEditToolAction>()->GetNameStringByIndex(static_cast<int>(ActionType));
+			Attributes.Add(FAnalyticsEventAttribute(FString::Printf(TEXT("Stats.%sAction.NumActions"), *ActionName), NumActions));
+		}
+	};
+	MaybeAppendStatsToAttributes(EChannelEditToolAction::Add);
+	MaybeAppendStatsToAttributes(EChannelEditToolAction::Delete);
+	MaybeAppendStatsToAttributes(EChannelEditToolAction::Copy);
+	Attributes.Add(FAnalyticsEventAttribute(TEXT("Stats.ToolActiveDuration"), (FDateTime::UtcNow() - ToolStartTimeAnalytics).ToString()));
+	
+	FEngineAnalytics::GetProvider().RecordEvent(UVEditorAnalyticsEventName(TEXT("ChannelEditTool")), Attributes);
+
+	if constexpr (false)
+	{
+		for (const FAnalyticsEventAttribute& Attr : Attributes)
+		{
+			UE_LOG(LogGeometry, Log, TEXT("Debug %s.ChannelEditTool.%s = %s"), *UVEditorAnalyticsPrefix, *Attr.GetName(), *Attr.GetValue());
+		}
+	}
+}
+
 #undef LOCTEXT_NAMESPACE
\ No newline at end of file
diff --git a/Engine/Plugins/Experimental/UVEditor/Source/UVEditorTools/Private/UVEditorLayoutTool.cpp b/Engine/Plugins/Experimental/UVEditor/Source/UVEditorTools/Private/UVEditorLayoutTool.cpp
--- a/Engine/Plugins/Experimental/UVEditor/Source/UVEditorTools/Private/UVEditorLayoutTool.cpp
+++ b/Engine/Plugins/Experimental/UVEditor/Source/UVEditorTools/Private/UVEditorLayoutTool.cpp
@@ -11,6 +11,8 @@
 #include "Properties/UVLayoutProperties.h"
 #include "ToolTargets/UVEditorToolMeshInput.h"
 #include "UVToolContextObjects.h"
+#include "EngineAnalytics.h"
+#include "UVEditorToolAnalyticsUtils.h"
 
 using namespace UE::Geometry;
 
@@ -36,6 +38,8 @@ UInteractiveTool* UUVEditorLayoutToolBuilder::BuildTool(const FToolBuilderState&
 void UUVEditorLayoutTool::Setup()
 {
 	check(Targets.Num() > 0);
+	
+	ToolStartTimeAnalytics = FDateTime::UtcNow();
 
 	UInteractiveTool::Setup();
 
@@ -64,6 +68,9 @@ void UUVEditorLayoutTool::Setup()
 	SetToolDisplayName(LOCTEXT("ToolName", "UV Layout"));
 	GetToolManager()->DisplayMessage(LOCTEXT("OnStartUVLayoutTool", "Translate, rotate or scale existing UV Charts using various strategies"),
 		EToolMessageLevel::UserNotification);
+
+	// Analytics
+	InputTargetAnalytics = UVEditorAnalytics::CollectTargetAnalytics(Targets);
 }
 
 void UUVEditorLayoutTool::Shutdown(EToolShutdownType ShutdownType)
@@ -101,6 +108,9 @@ void UUVEditorLayoutTool::Shutdown(EToolShutdownType ShutdownType)
 		}
 
 		ChangeAPI->EndUndoTransaction();
+
+		// Analytics
+		RecordAnalytics();
 	}
 	else
 	{
@@ -128,6 +138,10 @@ void UUVEditorLayoutTool::Shutdown(EToolShutdownType ShutdownType)
 
 void UUVEditorLayoutTool::OnTick(float DeltaTime)
 {
+	for (TObjectPtr<UUVEditorToolMeshInput> Target : Targets)
+	{
+		Target->AppliedPreview->Tick(DeltaTime);
+	}
 }
 
 
@@ -152,4 +166,51 @@ bool UUVEditorLayoutTool::CanAccept() const
 	return true;
 }
 
+void UUVEditorLayoutTool::RecordAnalytics()
+{
+	using namespace UVEditorAnalytics;
+	
+	if (!FEngineAnalytics::IsAvailable())
+	{
+		return;
+	}
+	
+	TArray<FAnalyticsEventAttribute> Attributes;
+	Attributes.Add(FAnalyticsEventAttribute(TEXT("Timestamp"), FDateTime::UtcNow().ToString()));
+
+	// Tool inputs
+	InputTargetAnalytics.AppendToAttributes(Attributes, "Input");
+
+	// Tool stats
+	if (CanAccept())
+	{
+		TArray<double> PerAssetValidResultComputeTimes;
+		for (TObjectPtr<UUVEditorToolMeshInput> Target : Targets)
+		{
+			// Note: This would log -1 if the result was invalid, but checking CanAccept above ensures results are valid
+			PerAssetValidResultComputeTimes.Add(Target->AppliedPreview->GetValidResultComputeTime());
+		}
+		Attributes.Add(FAnalyticsEventAttribute(TEXT("Stats.PerAsset.ComputeTimeSeconds"), PerAssetValidResultComputeTimes));
+	}
+	Attributes.Add(FAnalyticsEventAttribute(TEXT("Stats.ToolActiveDuration"), (FDateTime::UtcNow() - ToolStartTimeAnalytics).ToString()));
+
+	// Tool settings chosen by the user
+	Attributes.Add(AnalyticsEventAttributeEnum(TEXT("Settings.LayoutType"), Settings->LayoutType));
+	Attributes.Add(FAnalyticsEventAttribute(TEXT("Settings.TextureResolution"), Settings->TextureResolution));
+	Attributes.Add(FAnalyticsEventAttribute(TEXT("Settings.Scale"), Settings->Scale));
+	const TArray<FVector2D::FReal> TranslationArray({Settings->Translation.X, Settings->Translation.Y});
+	Attributes.Add(FAnalyticsEventAttribute(TEXT("Settings.Translation"), TranslationArray));
+	Attributes.Add(FAnalyticsEventAttribute(TEXT("Settings.AllowFlips"), Settings->bAllowFlips));
+
+	FEngineAnalytics::GetProvider().RecordEvent(UVEditorAnalyticsEventName(TEXT("LayoutTool")), Attributes);
+
+	if constexpr (false)
+	{
+		for (const FAnalyticsEventAttribute& Attr : Attributes)
+		{
+			UE_LOG(LogGeometry, Log, TEXT("Debug %s.LayoutTool.%s = %s"), *UVEditorAnalyticsPrefix, *Attr.GetName(), *Attr.GetValue());
+		}
+	}
+}
+
 #undef LOCTEXT_NAMESPACE
diff --git a/Engine/Plugins/Experimental/UVEditor/Source/UVEditorTools/Private/UVEditorRecomputeUVsTool.cpp b/Engine/Plugins/Experimental/UVEditor/Source/UVEditorTools/Private/UVEditorRecomputeUVsTool.cpp
--- a/Engine/Plugins/Experimental/UVEditor/Source/UVEditorTools/Private/UVEditorRecomputeUVsTool.cpp
+++ b/Engine/Plugins/Experimental/UVEditor/Source/UVEditorTools/Private/UVEditorRecomputeUVsTool.cpp
@@ -10,6 +10,8 @@
 #include "UVToolContextObjects.h"
 #include "ContextObjectStore.h"
 #include "MeshOpPreviewHelpers.h"
+#include "UVEditorToolAnalyticsUtils.h"
+#include "EngineAnalytics.h"
 
 using namespace UE::Geometry;
 
@@ -41,6 +43,8 @@ UInteractiveTool* UUVEditorRecomputeUVsToolBuilder::BuildTool(const FToolBuilder
 
 void UUVEditorRecomputeUVsTool::Setup()
 {
+	ToolStartTimeAnalytics = FDateTime::UtcNow();
+	
 	UInteractiveTool::Setup();
 
 	// initialize our properties
@@ -82,6 +86,9 @@ void UUVEditorRecomputeUVsTool::Setup()
 	GetToolManager()->DisplayMessage(
 		LOCTEXT("OnStartTool_Regions", "Generate UVs for PolyGroups or existing UV islands of the mesh using various strategies."),
 		EToolMessageLevel::UserNotification);
+	
+	// Analytics
+	InputTargetAnalytics = UVEditorAnalytics::CollectTargetAnalytics(Targets);
 }
 
 
@@ -128,6 +135,9 @@ void UUVEditorRecomputeUVsTool::Shutdown(EToolShutdownType ShutdownType)
 		}
 
 		ChangeAPI->EndUndoTransaction();
+
+		// Analytics
+		RecordAnalytics();
 	}
 	else
 	{
@@ -150,6 +160,14 @@ void UUVEditorRecomputeUVsTool::Shutdown(EToolShutdownType ShutdownType)
 	Targets.Empty();
 }
 
+void UUVEditorRecomputeUVsTool::OnTick(float DeltaTime)
+{
+	for (TObjectPtr<UUVEditorToolMeshInput> Target : Targets)
+	{
+		Target->AppliedPreview->Tick(DeltaTime);
+	}
+}
+
 void UUVEditorRecomputeUVsTool::Render(IToolsContextRenderAPI* RenderAPI)
 {
 }
@@ -196,4 +214,74 @@ void UUVEditorRecomputeUVsTool::UpdateActiveGroupLayer()
 }
 
 
+void UUVEditorRecomputeUVsTool::RecordAnalytics()
+{
+	using namespace UVEditorAnalytics;
+	
+	if (!FEngineAnalytics::IsAvailable())
+	{
+		return;
+	}
+
+	TArray<FAnalyticsEventAttribute> Attributes;
+	Attributes.Add(FAnalyticsEventAttribute(TEXT("Timestamp"), FDateTime::UtcNow().ToString()));
+
+	// Tool inputs
+	InputTargetAnalytics.AppendToAttributes(Attributes, "Input");
+
+	// Tool outputs	
+	const FTargetAnalytics OutputTargetAnalytics = CollectTargetAnalytics(Targets);
+	OutputTargetAnalytics.AppendToAttributes(Attributes, "Output");
+	
+	// Tool stats
+	if (CanAccept())
+	{
+		TArray<double> PerAssetValidResultComputeTimes;
+		for (TObjectPtr<UUVEditorToolMeshInput> Target : Targets)
+		{
+			// Note: This would log -1 if the result was invalid, but checking CanAccept above ensures results are valid
+			PerAssetValidResultComputeTimes.Add(Target->AppliedPreview->GetValidResultComputeTime());
+		}
+		Attributes.Add(FAnalyticsEventAttribute(TEXT("Stats.PerAsset.ComputeTimeSeconds"), PerAssetValidResultComputeTimes));
+	}
+	Attributes.Add(FAnalyticsEventAttribute(TEXT("Stats.ToolActiveDuration"), (FDateTime::UtcNow() - ToolStartTimeAnalytics).ToString()));
+	
+	// Tool settings chosen by the user (Volatile! Sync with EditCondition meta-tags in URecomputeUVsToolProperties)
+	
+	Attributes.Add(AnalyticsEventAttributeEnum(TEXT("Settings.IslandGeneration"), Settings->IslandGeneration));
+	Attributes.Add(AnalyticsEventAttributeEnum(TEXT("Settings.AutoRotation"), Settings->AutoRotation));
+	
+	Attributes.Add(AnalyticsEventAttributeEnum(TEXT("Settings.UnwrapType"), Settings->UnwrapType));
+	if (Settings->UnwrapType == ERecomputeUVsPropertiesUnwrapType::IslandMerging || Settings->UnwrapType == ERecomputeUVsPropertiesUnwrapType::ExpMap)
+	{
+		Attributes.Add(FAnalyticsEventAttribute(TEXT("Settings.SmoothingSteps"), Settings->SmoothingSteps));
+		Attributes.Add(FAnalyticsEventAttribute(TEXT("Settings.SmoothingAlpha"), Settings->SmoothingAlpha));
+	}
+	if (Settings->UnwrapType == ERecomputeUVsPropertiesUnwrapType::IslandMerging)
+	{
+		Attributes.Add(FAnalyticsEventAttribute(TEXT("Settings.MergingDistortionLimit"), Settings->MergingDistortionLimit));
+		Attributes.Add(FAnalyticsEventAttribute(TEXT("Settings.MergingAngleLimit"), Settings->MergingAngleLimit));
+	}
+	
+	Attributes.Add(AnalyticsEventAttributeEnum(TEXT("Settings.LayoutType"), Settings->LayoutType));
+	if (Settings->LayoutType == ERecomputeUVsPropertiesLayoutType::Repack)
+	{
+		Attributes.Add(FAnalyticsEventAttribute(TEXT("Settings.TextureResolution"), Settings->TextureResolution));
+	}
+	if (Settings->LayoutType == ERecomputeUVsPropertiesLayoutType::NormalizeToBounds || Settings->LayoutType == ERecomputeUVsPropertiesLayoutType::NormalizeToWorld)
+	{
+		Attributes.Add(FAnalyticsEventAttribute(TEXT("Settings.NormalizeScale"), Settings->NormalizeScale));
+	}
+
+	FEngineAnalytics::GetProvider().RecordEvent(UVEditorAnalyticsEventName(TEXT("UnwrapTool")), Attributes);
+
+	if constexpr (false)
+	{
+		for (const FAnalyticsEventAttribute& Attr : Attributes)
+		{
+			UE_LOG(LogGeometry, Log, TEXT("Debug %s.UnwrapTool.%s = %s"), *UVEditorAnalyticsPrefix, *Attr.GetName(), *Attr.GetValue());
+		}
+	}
+}
+
 #undef LOCTEXT_NAMESPACE
diff --git a/Engine/Plugins/Experimental/UVEditor/Source/UVEditorTools/Private/UVEditorSeamTool.cpp b/Engine/Plugins/Experimental/UVEditor/Source/UVEditorTools/Private/UVEditorSeamTool.cpp
--- a/Engine/Plugins/Experimental/UVEditor/Source/UVEditorTools/Private/UVEditorSeamTool.cpp
+++ b/Engine/Plugins/Experimental/UVEditor/Source/UVEditorTools/Private/UVEditorSeamTool.cpp
@@ -17,6 +17,7 @@
 #include "ToolSceneQueriesUtil.h"
 #include "UVEditorUXSettings.h"
 #include "UVToolContextObjects.h"
+#include "EngineAnalytics.h"
 
 #define LOCTEXT_NAMESPACE "UUVEditorSeamTool"
 
@@ -286,6 +287,8 @@ void UUVEditorSeamTool::SetTargets(const TArray<TObjectPtr<UUVEditorToolMeshInpu
 void UUVEditorSeamTool::Setup()
 {
 	TRACE_CPUPROFILER_EVENT_SCOPE(UVEditorSeamTool_Setup);
+	
+	ToolStartTimeAnalytics = FDateTime::UtcNow();
 
 	using namespace UVEditorSeamToolLocals;
 
@@ -463,6 +466,9 @@ void UUVEditorSeamTool::Setup()
 	}
 
 	UpdateToolMessage();
+
+	// Analytics
+	InputTargetAnalytics = UVEditorAnalytics::CollectTargetAnalytics(Targets);
 }
 
 void UUVEditorSeamTool::ReconstructExistingSeamsVisualization()
@@ -589,6 +595,9 @@ void UUVEditorSeamTool::Shutdown(EToolShutdownType ShutdownType)
 
 	Spatials2D.Reset();
 	Spatials3D.Reset();
+
+	// Analytics
+	RecordAnalytics();
 }
 
 void UUVEditorSeamTool::OnTick(float DeltaTime)
@@ -1173,4 +1182,37 @@ bool UUVEditorSeamTool::ExecuteNestedAcceptCommand()
 	return false;
 }
 
+void UUVEditorSeamTool::RecordAnalytics()
+{
+	using namespace UVEditorAnalytics;
+	
+	if (!FEngineAnalytics::IsAvailable())
+	{
+		return;
+	}
+	
+	TArray<FAnalyticsEventAttribute> Attributes;
+	Attributes.Add(FAnalyticsEventAttribute(TEXT("Timestamp"), FDateTime::UtcNow().ToString()));
+	
+	// Tool inputs
+	InputTargetAnalytics.AppendToAttributes(Attributes, "Input");
+
+	// Tool outputs
+	const FTargetAnalytics OutputTargetAnalytics = CollectTargetAnalytics(Targets);
+	OutputTargetAnalytics.AppendToAttributes(Attributes, "Output");
+
+	// Tool stats
+	Attributes.Add(FAnalyticsEventAttribute(TEXT("Stats.ToolActiveDuration"), (FDateTime::UtcNow() - ToolStartTimeAnalytics).ToString()));
+	
+	FEngineAnalytics::GetProvider().RecordEvent(UVEditorAnalyticsEventName(TEXT("SeamTool")), Attributes);
+
+	if constexpr (false)
+	{
+		for (const FAnalyticsEventAttribute& Attr : Attributes)
+		{
+			UE_LOG(LogGeometry, Log, TEXT("Debug %s.SeamTool.%s = %s"), *UVEditorAnalyticsPrefix, *Attr.GetName(), *Attr.GetValue());
+		}
+	}
+}
+
 #undef LOCTEXT_NAMESPACE
diff --git a/Engine/Plugins/Experimental/UVEditor/Source/UVEditorTools/Private/UVEditorToolAnalyticsUtils.cpp b/Engine/Plugins/Experimental/UVEditor/Source/UVEditorTools/Private/UVEditorToolAnalyticsUtils.cpp
new file mode 100644
--- /dev/null
+++ b/Engine/Plugins/Experimental/UVEditor/Source/UVEditorTools/Private/UVEditorToolAnalyticsUtils.cpp
@@ -0,0 +1,69 @@
+﻿// Copyright Epic Games, Inc. All Rights Reserved.
+
+#include "UVEditorToolAnalyticsUtils.h"
+
+#include "UObject/ObjectPtr.h"
+#include "ToolTargets/UVEditorToolMeshInput.h"
+
+namespace UE
+{
+namespace Geometry
+{
+	
+namespace UVEditorAnalytics
+{
+
+void FTargetAnalytics::AppendToAttributes(TArray<FAnalyticsEventAttribute>& Attributes, FString Prefix) const
+{
+	if (!Prefix.IsEmpty())
+	{
+		Prefix += ".";
+	}
+	
+	Attributes.Add(FAnalyticsEventAttribute(Prefix + TEXT("AllAssets.NumVertices"), AllAssetsNumVertices));
+	Attributes.Add(FAnalyticsEventAttribute(Prefix + TEXT("AllAssets.NumTriangles"), AllAssetsNumTriangles));
+	Attributes.Add(FAnalyticsEventAttribute(Prefix + TEXT("AllAssets.NumUVLayers"), AllAssetsNumUVLayers));
+	
+	Attributes.Add(FAnalyticsEventAttribute(Prefix + TEXT("PerAsset.NumVertices"), PerAssetNumVertices));
+	Attributes.Add(FAnalyticsEventAttribute(Prefix + TEXT("PerAsset.NumTriangles"), PerAssetNumTriangles));
+	Attributes.Add(FAnalyticsEventAttribute(Prefix + TEXT("PerAsset.NumUVLayers"), PerAssetNumUVLayers));
+	Attributes.Add(FAnalyticsEventAttribute(Prefix + TEXT("PerAsset.ActiveUVChannel.Index"), PerAssetActiveUVChannelIndex));
+	Attributes.Add(FAnalyticsEventAttribute(Prefix + TEXT("PerAsset.ActiveUVChannel.NumVertices"), PerAssetActiveUVChannelNumVertices));
+	Attributes.Add(FAnalyticsEventAttribute(Prefix + TEXT("PerAsset.ActiveUVChannel.NumTriangles"), PerAssetActiveUVChannelNumTriangles));
+}
+
+FTargetAnalytics CollectTargetAnalytics(const TArray<TObjectPtr<UUVEditorToolMeshInput>>& Targets)
+{
+	FTargetAnalytics Result;
+	
+	for (const TObjectPtr<UUVEditorToolMeshInput>& Target : Targets)
+	{
+		const int NumVertices = Target->AppliedCanonical->VertexCount();
+		const int NumTriangles = Target->AppliedCanonical->TriangleCount();
+		const int NumUVLayers = Target->AppliedCanonical->HasAttributes() ? Target->AppliedCanonical->Attributes()->NumUVLayers() : 0;
+		
+		Result.AllAssetsNumVertices += NumVertices;
+		Result.AllAssetsNumTriangles += NumTriangles;
+		Result.AllAssetsNumUVLayers += NumUVLayers;
+		
+		Result.PerAssetNumVertices.Add(NumVertices);
+		Result.PerAssetNumTriangles.Add(NumTriangles);
+		Result.PerAssetNumUVLayers.Add(NumUVLayers);
+
+		Result.PerAssetActiveUVChannelIndex.Add(Target->UVLayerIndex);
+		if (Target->UnwrapCanonical)
+		{
+			const int NumVerticesActiveUVChannel = Target->UnwrapCanonical->VertexCount();
+			const int NumTrianglesActiveUVChannel = Target->UnwrapCanonical->TriangleCount();
+			Result.PerAssetActiveUVChannelNumVertices.Add(NumVerticesActiveUVChannel);
+			Result.PerAssetActiveUVChannelNumTriangles.Add(NumTrianglesActiveUVChannel);
+		}
+	}
+
+	return Result;
+}
+
+} // end namespace UVEditorAnalytics
+	
+} // end namespace Geometry
+} // end namespace UE
diff --git a/Engine/Plugins/Experimental/UVEditor/Source/UVEditorTools/Private/UVSelectTool.cpp b/Engine/Plugins/Experimental/UVEditor/Source/UVEditorTools/Private/UVSelectTool.cpp
--- a/Engine/Plugins/Experimental/UVEditor/Source/UVEditorTools/Private/UVSelectTool.cpp
+++ b/Engine/Plugins/Experimental/UVEditor/Source/UVEditorTools/Private/UVSelectTool.cpp
@@ -21,6 +21,7 @@
 #include "Selection/DynamicMeshSelection.h"
 #include "ToolSetupUtil.h"
 #include "UVEditorUXSettings.h"
+#include "EngineAnalytics.h"
 
 #include "UVSeamSewAction.h"
 #include "UVIslandConformalUnwrapAction.h"
@@ -284,6 +285,8 @@ void USelectToolActionPropertySet::PostAction(ESelectToolAction Action)
 void UUVSelectTool::Setup()
 {
 	TRACE_CPUPROFILER_EVENT_SCOPE(UVSelectTool_Setup);
+	
+	ToolStartTimeAnalytics = FDateTime::UtcNow();
 
 	using namespace UVSelectToolLocals;
 
@@ -484,6 +487,9 @@ void UUVSelectTool::Setup()
 	GetToolManager()->DisplayMessage(LOCTEXT("SelectToolStatusBarMessage", 
 		"Select elements in the viewport and then transform them or use one of the action buttons."), 
 		EToolMessageLevel::UserNotification);
+
+	// Analytics
+	InputTargetAnalytics = UVEditorAnalytics::CollectTargetAnalytics(Targets);
 }
 
 void UUVSelectTool::Shutdown(EToolShutdownType ShutdownType)
@@ -535,6 +541,9 @@ void UUVSelectTool::Shutdown(EToolShutdownType ShutdownType)
 	ViewportButtonsAPI = nullptr;
 	EmitChangeAPI = nullptr;
 	ChangeRouter = nullptr;
+
+	// Analytics
+	RecordAnalytics();
 }
 
 void UUVSelectTool::SetSelection(const FDynamicMeshSelection& NewSelection)
@@ -964,12 +973,26 @@ void UUVSelectTool::RequestAction(ESelectToolAction ActionType)
 
 void UUVSelectTool::ApplyAction(ESelectToolAction ActionType)
 {
+	auto MaybeAddAnalyticsActionHistoryItem = [this, &ActionType](FDynamicMeshSelection::EType ExpectedType)
+	{
+		if (!SelectionMechanic->GetCurrentSelection().IsEmpty() && (SelectionMechanic->GetCurrentSelection().Type == ExpectedType))
+		{
+			FActionHistoryItem ActionHistoryItem;
+			ActionHistoryItem.Timestamp = FDateTime::UtcNow();
+			ActionHistoryItem.NumOperands = SelectionMechanic->GetCurrentSelection().SelectedIDs.Num();
+			ActionHistoryItem.ActionType = ActionType;
+			AnalyticsActionHistory.Add(ActionHistoryItem);
+		}
+	};
+	
 	switch (ActionType)
 	{
 	case ESelectToolAction::Sew:
 		if (SewAction)
 		{
 			TRACE_CPUPROFILER_EVENT_SCOPE(ApplyAction_Sew);
+	
+			MaybeAddAnalyticsActionHistoryItem(FDynamicMeshSelection::EType::Edge);
 
 			const FText TransactionName(LOCTEXT("SewCompleteTransactionName", "Sew Edges"));
 			EmitChangeAPI->BeginUndoTransaction(TransactionName);
@@ -984,6 +1007,8 @@ void UUVSelectTool::ApplyAction(ESelectToolAction ActionType)
 		if (IslandConformalUnwrapAction)
 		{
 			TRACE_CPUPROFILER_EVENT_SCOPE(ApplyAction_IslandConformalUnwrap);
+			
+			MaybeAddAnalyticsActionHistoryItem(FDynamicMeshSelection::EType::Triangle);
 
 			const FText TransactionName(LOCTEXT("ConformalUnwrapCompleteTransactionName", "Conformal Unwrap Islands"));
 			EmitChangeAPI->BeginUndoTransaction(TransactionName);
@@ -995,6 +1020,9 @@ void UUVSelectTool::ApplyAction(ESelectToolAction ActionType)
 		}
 		break;
 	case ESelectToolAction::Split:
+		
+		MaybeAddAnalyticsActionHistoryItem(FDynamicMeshSelection::EType::Edge);
+		
 		ApplySplit();
 	default:
 		break;
@@ -1162,4 +1190,68 @@ void UUVSelectTool::ApplySplitBowtieVertices()
 	EmitChangeAPI->EndUndoTransaction();
 }
 
+void UUVSelectTool::RecordAnalytics()
+{
+	using namespace UVEditorAnalytics;
+	
+	if (!FEngineAnalytics::IsAvailable())
+	{
+		return;
+	}
+	
+	TArray<FAnalyticsEventAttribute> Attributes;
+	Attributes.Add(FAnalyticsEventAttribute(TEXT("Timestamp"), FDateTime::UtcNow().ToString()));
+	
+	// Tool inputs
+	InputTargetAnalytics.AppendToAttributes(Attributes, "Input");
+
+	// Tool outputs
+	const FTargetAnalytics OutputTargetAnalytics = CollectTargetAnalytics(Targets);
+	OutputTargetAnalytics.AppendToAttributes(Attributes, "Output");
+
+	// Tool stats
+	auto MaybeAppendActionStatsToAttributes = [this, &Attributes](ESelectToolAction ActionType, const FString& OperandName)
+	{
+		int32 NumActions = 0;
+		float MeanNumOperands = 0;
+		int32 MinNumOperands = TNumericLimits<int32>::Max();
+		int32 MaxNumOperands = TNumericLimits<int32>::Min();
+			
+		for (const FActionHistoryItem& Item : AnalyticsActionHistory)
+		{
+			if (Item.ActionType == ActionType)
+			{
+				NumActions += 1;
+				MeanNumOperands += static_cast<float>(Item.NumOperands);
+				MinNumOperands = FGenericPlatformMath::Min(MinNumOperands, Item.NumOperands);
+				MaxNumOperands = FGenericPlatformMath::Max(MaxNumOperands, Item.NumOperands);
+			}
+		}
+		
+		if (NumActions > 0)
+		{
+			MeanNumOperands /= NumActions;
+			const FString ActionName = StaticEnum<ESelectToolAction>()->GetNameStringByIndex(static_cast<int>(ActionType));
+			Attributes.Add(FAnalyticsEventAttribute(FString::Printf(TEXT("Stats.%sAction.NumActions"), *ActionName), NumActions));
+			Attributes.Add(FAnalyticsEventAttribute(FString::Printf(TEXT("Stats.%sAction.MinNum%s"), *ActionName, *OperandName), MinNumOperands));
+			Attributes.Add(FAnalyticsEventAttribute(FString::Printf(TEXT("Stats.%sAction.MaxNum%s"), *ActionName, *OperandName), MaxNumOperands));
+			Attributes.Add(FAnalyticsEventAttribute(FString::Printf(TEXT("Stats.%sAction.MeanNum%s"), *ActionName, *OperandName), MeanNumOperands));
+		}
+	};
+	MaybeAppendActionStatsToAttributes(ESelectToolAction::IslandConformalUnwrap, TEXT("Triangles"));
+	MaybeAppendActionStatsToAttributes(ESelectToolAction::Split, TEXT("Edges"));
+	MaybeAppendActionStatsToAttributes(ESelectToolAction::Sew, TEXT("Edges"));
+	Attributes.Add(FAnalyticsEventAttribute(TEXT("Stats.ToolActiveDuration"), (FDateTime::UtcNow() - ToolStartTimeAnalytics).ToString()));
+	
+	FEngineAnalytics::GetProvider().RecordEvent(UVEditorAnalyticsEventName(TEXT("EditTool")), Attributes);
+
+	if constexpr (false)
+	{
+		for (const FAnalyticsEventAttribute& Attr : Attributes)
+		{
+			UE_LOG(LogGeometry, Log, TEXT("Debug %s.EditTool.%s = %s"), *UVEditorAnalyticsPrefix, *Attr.GetName(), *Attr.GetValue());
+		}
+	}
+}
+
 #undef LOCTEXT_NAMESPACE
diff --git a/Engine/Plugins/Experimental/UVEditor/Source/UVEditorTools/Public/UVEditorLayerEditTool.h b/Engine/Plugins/Experimental/UVEditor/Source/UVEditorTools/Public/UVEditorLayerEditTool.h
--- a/Engine/Plugins/Experimental/UVEditor/Source/UVEditorTools/Public/UVEditorLayerEditTool.h
+++ b/Engine/Plugins/Experimental/UVEditor/Source/UVEditorTools/Public/UVEditorLayerEditTool.h
@@ -6,6 +6,7 @@
 
 #include "InteractiveTool.h"
 #include "InteractiveToolBuilder.h"
+#include "UVEditorToolAnalyticsUtils.h"
 
 #include "UVEditorLayerEditTool.generated.h"
 
@@ -228,4 +229,30 @@ protected:
 	UPROPERTY()
 	TObjectPtr<UUVToolEmitChangeAPI> EmitChangeAPI = nullptr;
 
+	//
+	// Analytics
+	//
+
+	struct FActionHistoryItem
+	{
+		FDateTime Timestamp;
+		EChannelEditToolAction ActionType = EChannelEditToolAction::NoAction;
+
+		// if ActionType == Add    then FirstOperandIndex is the index of the added UV layer
+		// if ActionType == Delete then FirstOperandIndex is the index of the deleted UV layer
+		// if ActionType == Copy   then FirstOperandIndex is the index of the source UV layer
+		int32 FirstOperandIndex = -1;
+		
+		// if ActionType == Add    then SecondOperandIndex is unused
+		// if ActionType == Delete then SecondOperandIndex is unused
+		// if ActionType == Copy   then SecondOperandIndex is the index of the target UV layer
+		int32 SecondOperandIndex = -1;
+
+		bool bDeleteActionWasActuallyClear = false;
+	};
+	
+	TArray<FActionHistoryItem> AnalyticsActionHistory;
+	UE::Geometry::UVEditorAnalytics::FTargetAnalytics InputTargetAnalytics;
+	FDateTime ToolStartTimeAnalytics;
+	void RecordAnalytics();
 };
diff --git a/Engine/Plugins/Experimental/UVEditor/Source/UVEditorTools/Public/UVEditorLayoutTool.h b/Engine/Plugins/Experimental/UVEditor/Source/UVEditorTools/Public/UVEditorLayoutTool.h
--- a/Engine/Plugins/Experimental/UVEditor/Source/UVEditorTools/Public/UVEditorLayoutTool.h
+++ b/Engine/Plugins/Experimental/UVEditor/Source/UVEditorTools/Public/UVEditorLayoutTool.h
@@ -6,6 +6,7 @@
 
 #include "InteractiveTool.h"
 #include "InteractiveToolBuilder.h"
+#include "UVEditorToolAnalyticsUtils.h"
 
 #include "UVEditorLayoutTool.generated.h"
 
@@ -65,4 +66,12 @@ protected:
 
 	UPROPERTY()
 	TArray<TObjectPtr<UUVLayoutOperatorFactory>> Factories;
+
+	//
+	// Analytics
+	//
+
+	UE::Geometry::UVEditorAnalytics::FTargetAnalytics InputTargetAnalytics;
+	FDateTime ToolStartTimeAnalytics;
+	void RecordAnalytics();
 };
diff --git a/Engine/Plugins/Experimental/UVEditor/Source/UVEditorTools/Public/UVEditorRecomputeUVsTool.h b/Engine/Plugins/Experimental/UVEditor/Source/UVEditorTools/Public/UVEditorRecomputeUVsTool.h
--- a/Engine/Plugins/Experimental/UVEditor/Source/UVEditorTools/Public/UVEditorRecomputeUVsTool.h
+++ b/Engine/Plugins/Experimental/UVEditor/Source/UVEditorTools/Public/UVEditorRecomputeUVsTool.h
@@ -9,6 +9,7 @@
 #include "Properties/RecomputeUVsProperties.h"
 #include "Polygroups/PolygroupSet.h"
 #include "Drawing/UVLayoutPreview.h"
+#include "UVEditorToolAnalyticsUtils.h"
 
 #include "UVEditorRecomputeUVsTool.generated.h"
 
@@ -49,6 +50,8 @@ class UVEDITORTOOLS_API UUVEditorRecomputeUVsTool : public UInteractiveTool
 public:
 	virtual void Setup() override;
 	virtual void Shutdown(EToolShutdownType ShutdownType) override;
+	
+	virtual void OnTick(float DeltaTime) override;
 
 	virtual void Render(IToolsContextRenderAPI* RenderAPI) override;
 
@@ -83,4 +86,12 @@ protected:
 	TSharedPtr<UE::Geometry::FPolygroupSet, ESPMode::ThreadSafe> ActiveGroupSet;
 	void OnSelectedGroupLayerChanged();
 	void UpdateActiveGroupLayer();
+
+	//
+	// Analytics
+	//
+	
+	UE::Geometry::UVEditorAnalytics::FTargetAnalytics InputTargetAnalytics;
+	FDateTime ToolStartTimeAnalytics;
+	void RecordAnalytics();
 };
diff --git a/Engine/Plugins/Experimental/UVEditor/Source/UVEditorTools/Public/UVEditorSeamTool.h b/Engine/Plugins/Experimental/UVEditor/Source/UVEditorTools/Public/UVEditorSeamTool.h
--- a/Engine/Plugins/Experimental/UVEditor/Source/UVEditorTools/Public/UVEditorSeamTool.h
+++ b/Engine/Plugins/Experimental/UVEditor/Source/UVEditorTools/Public/UVEditorSeamTool.h
@@ -10,6 +10,7 @@
 #include "InteractiveToolBuilder.h"
 #include "InteractiveToolQueryInterfaces.h" // IInteractiveToolNestedAcceptCancelAPI
 #include "GeometryBase.h"
+#include "UVEditorToolAnalyticsUtils.h"
 
 #include "UVEditorSeamTool.generated.h"
 
@@ -145,4 +146,12 @@ protected:
 	// When true, the entire path is changed to the "completion" color to show
 	// that the next click will complete the path.
 	bool bCompletionColorOverride = false;
+	
+	//
+	// Analytics
+	//
+	
+	UE::Geometry::UVEditorAnalytics::FTargetAnalytics InputTargetAnalytics;
+	FDateTime ToolStartTimeAnalytics;
+	void RecordAnalytics();
 };
diff --git a/Engine/Plugins/Experimental/UVEditor/Source/UVEditorTools/Public/UVEditorToolAnalyticsUtils.h b/Engine/Plugins/Experimental/UVEditor/Source/UVEditorTools/Public/UVEditorToolAnalyticsUtils.h
new file mode 100644
--- /dev/null
+++ b/Engine/Plugins/Experimental/UVEditor/Source/UVEditorTools/Public/UVEditorToolAnalyticsUtils.h
@@ -0,0 +1,54 @@
+﻿// Copyright Epic Games, Inc. All Rights Reserved.
+
+#pragma once
+
+#include "CoreMinimal.h"
+#include "EngineAnalytics.h"
+
+class UUVEditorToolMeshInput;
+
+namespace UE
+{
+namespace Geometry
+{
+	
+namespace UVEditorAnalytics
+{
+	
+template <typename EnumType>
+FAnalyticsEventAttribute AnalyticsEventAttributeEnum(const FString& AttributeName, EnumType EnumValue)
+{
+	static_assert(TIsEnum<EnumType>::Value, "");
+	const FString EnumValueName = StaticEnum<EnumType>()->GetNameStringByValue(static_cast<int64>(EnumValue));
+	return FAnalyticsEventAttribute(AttributeName, EnumValueName);
+}
+
+static const FString UVEditorAnalyticsPrefix(TEXT("Editor.Usage.UVEditor"));
+
+FORCEINLINE FString UVEditorAnalyticsEventName(const FString& ToolEventName)
+{
+	return FString::Printf(TEXT("%s.%s"), *UVEditorAnalyticsPrefix, *ToolEventName);
+}
+
+struct UVEDITORTOOLS_API FTargetAnalytics
+{
+	int64 AllAssetsNumVertices = 0;
+	int64 AllAssetsNumTriangles = 0;
+	int64 AllAssetsNumUVLayers = 0;
+	
+	TArray<int32> PerAssetNumVertices;
+	TArray<int32> PerAssetNumTriangles;
+	TArray<int32> PerAssetNumUVLayers;
+	TArray<int32> PerAssetActiveUVChannelIndex;
+	TArray<int32> PerAssetActiveUVChannelNumVertices;
+	TArray<int32> PerAssetActiveUVChannelNumTriangles;
+
+	void AppendToAttributes(TArray<FAnalyticsEventAttribute>& Attributes, FString Prefix = "") const;
+};
+
+UVEDITORTOOLS_API FTargetAnalytics CollectTargetAnalytics(const TArray<TObjectPtr<UUVEditorToolMeshInput>>& Targets);
+	
+} // end namespace UVEditorAnalytics
+	
+} // end namespace Geometry
+} // end namespace UE
diff --git a/Engine/Plugins/Experimental/UVEditor/Source/UVEditorTools/Public/UVSelectTool.h b/Engine/Plugins/Experimental/UVEditor/Source/UVEditorTools/Public/UVSelectTool.h
--- a/Engine/Plugins/Experimental/UVEditor/Source/UVEditorTools/Public/UVSelectTool.h
+++ b/Engine/Plugins/Experimental/UVEditor/Source/UVEditorTools/Public/UVSelectTool.h
@@ -11,6 +11,7 @@
 #include "InteractiveToolBuilder.h"
 #include "TargetInterfaces/UVUnwrapDynamicMesh.h"
 #include "UVToolContextObjects.h"
+#include "UVEditorToolAnalyticsUtils.h"
 #include "IndexTypes.h"
 
 #include "UVSelectTool.generated.h"
@@ -210,6 +211,26 @@ protected:
 	void ApplySplitBowtieVertices();
 
 
+
+	//
+	// Analytics
+	//
+
+	struct FActionHistoryItem
+	{
+		FDateTime Timestamp;
+		ESelectToolAction ActionType;
+		
+		// if ActionType == IslandConformalUnwrap then NumOperands is #triangles in island selection
+		// if ActionType == Split                 then NumOperands is #edges in selection
+		// if ActionType == Sew                   then NumOperands is #edges in selection
+		int32 NumOperands = -1;
+	};
+	
+	TArray<FActionHistoryItem> AnalyticsActionHistory;
+	UE::Geometry::UVEditorAnalytics::FTargetAnalytics InputTargetAnalytics;
+	FDateTime ToolStartTimeAnalytics;
+	void RecordAnalytics();
 };
 
 /**
diff --git a/Engine/Plugins/Experimental/UVEditor/Source/UVEditorToolsEditorOnly/Private/UVEditorParameterizeMeshTool.cpp b/Engine/Plugins/Experimental/UVEditor/Source/UVEditorToolsEditorOnly/Private/UVEditorParameterizeMeshTool.cpp
--- a/Engine/Plugins/Experimental/UVEditor/Source/UVEditorToolsEditorOnly/Private/UVEditorParameterizeMeshTool.cpp
+++ b/Engine/Plugins/Experimental/UVEditor/Source/UVEditorToolsEditorOnly/Private/UVEditorParameterizeMeshTool.cpp
@@ -12,6 +12,7 @@
 #include "ToolTargets/UVEditorToolMeshInput.h"
 #include "UVToolContextObjects.h"
 #include "ContextObjectStore.h"
+#include "EngineAnalytics.h"
 
 using namespace UE::Geometry;
 
@@ -43,6 +44,8 @@ UInteractiveTool* UUVEditorParameterizeMeshToolBuilder::BuildTool(const FToolBui
 void UUVEditorParameterizeMeshTool::Setup()
 {
 	TRACE_CPUPROFILER_EVENT_SCOPE(UVEditorParameterizeMeshTool_Setup);
+	
+	ToolStartTimeAnalytics = FDateTime::UtcNow();
 
 	check(Targets.Num() >= 1);
 
@@ -94,6 +97,9 @@ void UUVEditorParameterizeMeshTool::Setup()
 	GetToolManager()->DisplayMessage(
 		LOCTEXT("OnStartTool_Global", "Automatically partition the selected Mesh into UV islands, flatten, and pack into a single UV chart"),
 		EToolMessageLevel::UserNotification);
+
+	// Analytics
+	InputTargetAnalytics = UVEditorAnalytics::CollectTargetAnalytics(Targets);
 }
 
 void UUVEditorParameterizeMeshTool::OnPropertyModified(UObject* PropertySet, FProperty* Property)
@@ -160,6 +166,9 @@ void UUVEditorParameterizeMeshTool::Shutdown(EToolShutdownType ShutdownType)
 		}
 
 		ChangeAPI->EndUndoTransaction();
+
+		// Analytics
+		RecordAnalytics();
 	}
 	else
 	{
@@ -187,6 +196,10 @@ void UUVEditorParameterizeMeshTool::Shutdown(EToolShutdownType ShutdownType)
 
 void UUVEditorParameterizeMeshTool::OnTick(float DeltaTime)
 {
+	for (TObjectPtr<UUVEditorToolMeshInput> Target : Targets)
+	{
+		Target->AppliedPreview->Tick(DeltaTime);
+	}
 }
 
 bool UUVEditorParameterizeMeshTool::CanAccept() const
@@ -201,5 +214,77 @@ bool UUVEditorParameterizeMeshTool::CanAccept() const
 	return true;
 }
 
+void UUVEditorParameterizeMeshTool::RecordAnalytics()
+{
+	using namespace UVEditorAnalytics;
+	
+	if (!FEngineAnalytics::IsAvailable())
+	{
+		return;
+	}
+	
+	TArray<FAnalyticsEventAttribute> Attributes;
+	Attributes.Add(FAnalyticsEventAttribute(TEXT("Timestamp"), FDateTime::UtcNow().ToString()));
+	
+	// Tool inputs
+	InputTargetAnalytics.AppendToAttributes(Attributes, "Input");
+
+	// Tool outputs
+	const FTargetAnalytics OutputTargetAnalytics = CollectTargetAnalytics(Targets);
+	OutputTargetAnalytics.AppendToAttributes(Attributes, "Output");
+
+	// Tool stats
+	if (CanAccept())
+	{
+		TArray<double> PerAssetValidResultComputeTimes;
+		for (TObjectPtr<UUVEditorToolMeshInput> Target : Targets)
+		{
+			// Note: This would log -1 if the result was invalid, but checking CanAccept above ensures results are valid
+			PerAssetValidResultComputeTimes.Add(Target->AppliedPreview->GetValidResultComputeTime());
+		}
+		Attributes.Add(FAnalyticsEventAttribute(TEXT("Stats.PerAsset.ComputeTimeSeconds"), PerAssetValidResultComputeTimes));
+	}
+	Attributes.Add(FAnalyticsEventAttribute(TEXT("Stats.ToolActiveDuration"), (FDateTime::UtcNow() - ToolStartTimeAnalytics).ToString()));
+
+	// Tool settings chosen by the user (Volatile! Sync with EditCondition meta-tags in *Properties members)
+	const FString MethodName = StaticEnum<EParameterizeMeshUVMethod>()->GetNameStringByIndex(static_cast<int>(Settings->Method));
+	Attributes.Add(FAnalyticsEventAttribute(TEXT("Settings.Method"), MethodName));
+	switch (Settings->Method)
+	{
+		case EParameterizeMeshUVMethod::PatchBuilder:
+			Attributes.Add(FAnalyticsEventAttribute(FString::Printf(TEXT("Settings.%s.InitialPatches"), *MethodName), PatchBuilderProperties->InitialPatches));
+			Attributes.Add(FAnalyticsEventAttribute(FString::Printf(TEXT("Settings.%s.CurvatureAlignment"), *MethodName), PatchBuilderProperties->CurvatureAlignment));
+			Attributes.Add(FAnalyticsEventAttribute(FString::Printf(TEXT("Settings.%s.MergingThreshold"), *MethodName), PatchBuilderProperties->MergingThreshold));
+			Attributes.Add(FAnalyticsEventAttribute(FString::Printf(TEXT("Settings.%s.MaxAngleDeviation"), *MethodName), PatchBuilderProperties->MaxAngleDeviation));
+			Attributes.Add(FAnalyticsEventAttribute(FString::Printf(TEXT("Settings.%s.SmoothingSteps"), *MethodName), PatchBuilderProperties->SmoothingSteps));
+			Attributes.Add(FAnalyticsEventAttribute(FString::Printf(TEXT("Settings.%s.SmoothingAlpha"), *MethodName), PatchBuilderProperties->SmoothingAlpha));
+			Attributes.Add(FAnalyticsEventAttribute(FString::Printf(TEXT("Settings.%s.AutoPack"), *MethodName), PatchBuilderProperties->bAutoPack));
+			if (PatchBuilderProperties->bAutoPack)
+			{
+				Attributes.Add(FAnalyticsEventAttribute(FString::Printf(TEXT("Settings.%s.TextureResolution"), *MethodName), PatchBuilderProperties->TextureResolution));
+			}
+			break;
+		case EParameterizeMeshUVMethod::UVAtlas:
+			Attributes.Add(FAnalyticsEventAttribute(FString::Printf(TEXT("Settings.%s.ChartStretch"), *MethodName), UVAtlasProperties->ChartStretch));
+			Attributes.Add(FAnalyticsEventAttribute(FString::Printf(TEXT("Settings.%s.NumCharts"), *MethodName), UVAtlasProperties->NumCharts));
+			break;
+		case EParameterizeMeshUVMethod::XAtlas:
+			Attributes.Add(FAnalyticsEventAttribute(FString::Printf(TEXT("Settings.%s.MaxIterations"), *MethodName), XAtlasProperties->MaxIterations));
+			break;
+		default:
+			break;
+	}
+	
+	FEngineAnalytics::GetProvider().RecordEvent(UVEditorAnalyticsEventName(TEXT("AutoUVTool")), Attributes);
+
+	if constexpr (false)
+	{
+		for (const FAnalyticsEventAttribute& Attr : Attributes)
+		{
+			UE_LOG(LogGeometry, Log, TEXT("Debug %s.AutoUVTool.%s = %s"), *UVEditorAnalyticsPrefix, *Attr.GetName(), *Attr.GetValue());
+		}
+	}
+}
+
 
 #undef LOCTEXT_NAMESPACE
diff --git a/Engine/Plugins/Experimental/UVEditor/Source/UVEditorToolsEditorOnly/Public/UVEditorParameterizeMeshTool.h b/Engine/Plugins/Experimental/UVEditor/Source/UVEditorToolsEditorOnly/Public/UVEditorParameterizeMeshTool.h
--- a/Engine/Plugins/Experimental/UVEditor/Source/UVEditorToolsEditorOnly/Public/UVEditorParameterizeMeshTool.h
+++ b/Engine/Plugins/Experimental/UVEditor/Source/UVEditorToolsEditorOnly/Public/UVEditorParameterizeMeshTool.h
@@ -11,6 +11,7 @@
 #include "MeshOpPreviewHelpers.h"
 #include "Properties/MeshMaterialProperties.h"
 #include "Properties/MeshUVChannelProperties.h"
+#include "UVEditorToolAnalyticsUtils.h"
 
 #include "UVEditorParameterizeMeshTool.generated.h"
 
@@ -89,5 +90,13 @@ protected:
 	TArray<TObjectPtr<UParameterizeMeshOperatorFactory>> Factories;
 
 	void OnMethodTypeChanged();
+
+	//
+	// Analytics
+	//
+	
+	UE::Geometry::UVEditorAnalytics::FTargetAnalytics InputTargetAnalytics;
+	FDateTime ToolStartTimeAnalytics;
+	void RecordAnalytics();
 };
 
diff --git a/Engine/Plugins/Runtime/MeshModelingToolset/Source/ModelingComponents/Private/MeshOpPreviewHelpers.cpp b/Engine/Plugins/Runtime/MeshModelingToolset/Source/ModelingComponents/Private/MeshOpPreviewHelpers.cpp
--- a/Engine/Plugins/Runtime/MeshModelingToolset/Source/ModelingComponents/Private/MeshOpPreviewHelpers.cpp
+++ b/Engine/Plugins/Runtime/MeshModelingToolset/Source/ModelingComponents/Private/MeshOpPreviewHelpers.cpp
@@ -108,8 +108,14 @@ void UMeshOpPreviewWithBackgroundCompute::Tick(float DeltaTime)
 
 void UMeshOpPreviewWithBackgroundCompute::UpdateResults()
 {
-	LastComputeStatus = BackgroundCompute ? BackgroundCompute->CheckStatus()
-		: EBackgroundComputeTaskStatus::NotComputing;
+	if (BackgroundCompute == nullptr)
+	{
+		LastComputeStatus = EBackgroundComputeTaskStatus::NotComputing;
+		return;
+	}
+
+	FBackgroundDynamicMeshComputeSource::FStatus Status = BackgroundCompute->CheckStatus();	
+	LastComputeStatus = Status.TaskStatus; 
 
 	if (LastComputeStatus == EBackgroundComputeTaskStatus::ValidResultAvailable
 		|| (bAllowDirtyResultUpdates && LastComputeStatus == EBackgroundComputeTaskStatus::DirtyResultAvailable))
@@ -129,6 +135,7 @@ void UMeshOpPreviewWithBackgroundCompute::UpdateResults()
 
 		PreviewMesh->SetVisible(bVisible);
 		bResultValid = (LastComputeStatus == EBackgroundComputeTaskStatus::ValidResultAvailable);
+		ValidResultComputeTimeSeconds = Status.ElapsedTime;
 
 		OnMeshUpdated.Broadcast(this);
 	}
diff --git a/Engine/Plugins/Runtime/MeshModelingToolset/Source/ModelingComponents/Public/MeshOpPreviewHelpers.h b/Engine/Plugins/Runtime/MeshModelingToolset/Source/ModelingComponents/Public/MeshOpPreviewHelpers.h
--- a/Engine/Plugins/Runtime/MeshModelingToolset/Source/ModelingComponents/Public/MeshOpPreviewHelpers.h
+++ b/Engine/Plugins/Runtime/MeshModelingToolset/Source/ModelingComponents/Public/MeshOpPreviewHelpers.h
@@ -110,6 +110,15 @@ public:
 	 */
 	bool HaveValidResult() const { return bResultValid; }
 
+	double GetValidResultComputeTime() const
+	{
+		if (HaveValidResult())
+		{
+			return ValidResultComputeTimeSeconds;
+		}
+		return -1;
+	}
+
 
 	/**
 	 * @return true if current PreviewMesh result is valid (no update actively being computed) and that mesh has at least one triangle
@@ -227,6 +236,7 @@ public:
 protected:
 	// state flag, if true then we have valid result
 	bool bResultValid = false;
+	double ValidResultComputeTimeSeconds = -1;
 
 	// Stored status of last compute, mainly so that we know when we should
 	// show the "busy" material.
@@ -381,7 +391,7 @@ protected:
 		check(BackgroundCompute);
 		if (BackgroundCompute)
 		{
-			UE::Geometry::EBackgroundComputeTaskStatus Status = BackgroundCompute->CheckStatus();
+			UE::Geometry::EBackgroundComputeTaskStatus Status = BackgroundCompute->CheckStatus().TaskStatus;
 			if (Status == UE::Geometry::EBackgroundComputeTaskStatus::ValidResultAvailable)
 			{
 				TUniquePtr<OperatorType> ResultOp = BackgroundCompute->ExtractResult();
diff --git a/Engine/Plugins/Runtime/MeshModelingToolset/Source/ModelingOperators/Public/BackgroundModelingComputeSource.h b/Engine/Plugins/Runtime/MeshModelingToolset/Source/ModelingOperators/Public/BackgroundModelingComputeSource.h
--- a/Engine/Plugins/Runtime/MeshModelingToolset/Source/ModelingOperators/Public/BackgroundModelingComputeSource.h
+++ b/Engine/Plugins/Runtime/MeshModelingToolset/Source/ModelingOperators/Public/BackgroundModelingComputeSource.h
@@ -143,6 +143,7 @@ protected:
 
 	double AccumTime = 0;
 	double LastStartTime = 0;
+	mutable double LastEndTime = 0; // mutable because it must be set in CheckStatus, which is const
 	double LastInvalidateTime = 0;
 
 	void StartNewCompute();
@@ -173,10 +174,24 @@ public:
 	 */
 	void NotifyActiveComputeInvalidated();
 
+	struct FStatus
+	{
+		EBackgroundComputeTaskStatus TaskStatus = EBackgroundComputeTaskStatus::NotComputing;
+
+		// if TaskStatus == ValidResultAvailable then this is the time spent computing the (valid) result
+		// if TaskStatus == DirtyResultAvailable then this is the time spent computing the (dirty) result
+		// if TaskStatus == Aborted              then this is the time spent computing the result before the task was aborted
+		// if TaskStatus == InProgress           then this is the time spent computing the result so far
+		// if TaskStatus == NotComputing         then this is the time spent not computing anything
+		double ElapsedTime = -1;
+	};
+
 	/**
 	 * Return status of the active background computation.
+	 * Warning: Status.ElapsedTime may not be accurate because Tick is not called with the correct DeltaTime argument in
+	 * Development builds when the user is interacting with a slider (as of 13th Dec 2021)
 	 */
-	EBackgroundComputeTaskStatus CheckStatus() const;
+	FStatus CheckStatus() const;
 
 	/**
 	 * @return The last computed Operator. This may only be called once, the caller then owns the Operator.
@@ -247,11 +262,10 @@ void TBackgroundModelingComputeSource<OpType, OpTypeFactory>::StartNewCompute()
 template<typename OpType, typename OpTypeFactory>
 void TBackgroundModelingComputeSource<OpType, OpTypeFactory>::NotifyActiveComputeInvalidated()
 {
-	LastInvalidateTime = AccumTime;
-
 	// switch to waiting-to-cancel state
 	if (TaskState == EBackgroundComputeTaskState::ComputingResult)
 	{
+		LastInvalidateTime = AccumTime;
 		TaskState = EBackgroundComputeTaskState::WaitingToCancel;
 	}
 
@@ -267,26 +281,41 @@ void TBackgroundModelingComputeSource<OpType, OpTypeFactory>::NotifyActiveComput
 
 
 template<typename OpType, typename OpTypeFactory>
-EBackgroundComputeTaskStatus TBackgroundModelingComputeSource<OpType, OpTypeFactory>::CheckStatus() const
+typename TBackgroundModelingComputeSource<OpType, OpTypeFactory>::FStatus
+TBackgroundModelingComputeSource<OpType, OpTypeFactory>::CheckStatus() const
 {
+	FStatus Status;
+
 	if (ActiveBackgroundTask == nullptr)
 	{
-		return EBackgroundComputeTaskStatus::NotComputing;
+		Status.TaskStatus = EBackgroundComputeTaskStatus::NotComputing;
+		Status.ElapsedTime = AccumTime - LastInvalidateTime;
 	}
 	else if (!ActiveBackgroundTask->IsDone())
 	{
-		return EBackgroundComputeTaskStatus::InProgress;
+		Status.TaskStatus = EBackgroundComputeTaskStatus::InProgress;
+		Status.ElapsedTime = AccumTime - LastStartTime;
 	}
 	else if (ActiveBackgroundTask->GetTask().IsAborted())
 	{
-		return EBackgroundComputeTaskStatus::Aborted;
+		Status.TaskStatus = EBackgroundComputeTaskStatus::Aborted;
+		Status.ElapsedTime = LastInvalidateTime - LastStartTime;
 	}
 	else
 	{
+		if (TaskState == EBackgroundComputeTaskState::ComputingResult)
+		{
+			LastEndTime = AccumTime;
+		}
+
 		// Task is done and not aborted, but we may be waiting to start a new one
-		return TaskState == EBackgroundComputeTaskState::WaitingToCancel ? EBackgroundComputeTaskStatus::DirtyResultAvailable
-			: EBackgroundComputeTaskStatus::ValidResultAvailable;
+		Status.TaskStatus = (TaskState == EBackgroundComputeTaskState::WaitingToCancel ?
+			EBackgroundComputeTaskStatus::DirtyResultAvailable :
+			EBackgroundComputeTaskStatus::ValidResultAvailable);
+		Status.ElapsedTime = LastEndTime - LastStartTime;
 	}
+
+	return Status;
 }
 
 
