Non-Unity Fix

diff --git a/Engine/Plugins/Experimental/UVEditor/Source/UVEditorTools/Private/UVEditorToolAnalyticsUtils.cpp b/Engine/Plugins/Experimental/UVEditor/Source/UVEditorTools/Private/UVEditorToolAnalyticsUtils.cpp
--- a/Engine/Plugins/Experimental/UVEditor/Source/UVEditorTools/Private/UVEditorToolAnalyticsUtils.cpp
+++ b/Engine/Plugins/Experimental/UVEditor/Source/UVEditorTools/Private/UVEditorToolAnalyticsUtils.cpp
@@ -3,7 +3,7 @@
 #include "UVEditorToolAnalyticsUtils.h"
 
 #include "UObject/ObjectPtr.h"
-#include "ToolTargets/UVEditorToolMeshInput.h"
+
 
 namespace UE
 {
diff --git a/Engine/Plugins/Experimental/UVEditor/Source/UVEditorTools/Public/UVEditorToolAnalyticsUtils.h b/Engine/Plugins/Experimental/UVEditor/Source/UVEditorTools/Public/UVEditorToolAnalyticsUtils.h
--- a/Engine/Plugins/Experimental/UVEditor/Source/UVEditorTools/Public/UVEditorToolAnalyticsUtils.h
+++ b/Engine/Plugins/Experimental/UVEditor/Source/UVEditorTools/Public/UVEditorToolAnalyticsUtils.h
@@ -4,8 +4,7 @@
 
 #include "CoreMinimal.h"
 #include "EngineAnalytics.h"
-
-class UUVEditorToolMeshInput;
+#include "ToolTargets/UVEditorToolMeshInput.h"
 
 namespace UE
 {
