UVEditor: Add notification messages to the editor when the user performs a channel edit operation for adding, copying, or deleting a channel.

#rnx

diff --git a/Engine/Plugins/Experimental/UVEditor/Source/UVEditorTools/Private/UVEditorLayerEditTool.cpp b/Engine/Plugins/Experimental/UVEditor/Source/UVEditorTools/Private/UVEditorLayerEditTool.cpp
--- a/Engine/Plugins/Experimental/UVEditor/Source/UVEditorTools/Private/UVEditorLayerEditTool.cpp
+++ b/Engine/Plugins/Experimental/UVEditor/Source/UVEditorTools/Private/UVEditorLayerEditTool.cpp
@@ -602,7 +602,8 @@ void UUVEditorChannelEditTool::ApplyVisbleChannelChange()
 
 void UUVEditorChannelEditTool::OnPropertyModified(UObject* PropertySet, FProperty* Property)
 {
-
+	GetToolManager()->DisplayMessage(LOCTEXT("OnStartUVChannelEditTool", "Add, copy or delete UV channels"),
+		EToolMessageLevel::UserNotification);
 }
 
 EChannelEditToolAction UUVEditorChannelEditTool::ActiveAction() const
@@ -651,6 +652,7 @@ void UUVEditorChannelEditTool::AddChannel()
 		SourceChannelProperties->TargetChannel = SourceChannelProperties->GetUVChannelNames()[NewChannelIndex];
 	}
 
+	GetToolManager()->DisplayMessage(LOCTEXT("AddChannelNotification", "New UV Channel added."), EToolMessageLevel::UserNotification);
 }
 
 void UUVEditorChannelEditTool::CopyChannel()
@@ -672,6 +674,8 @@ void UUVEditorChannelEditTool::CopyChannel()
 	Targets[ActiveAsset]->UpdateAllFromAppliedCanonical();
 
 	SourceChannelProperties->Initialize(Targets, false);
+
+	GetToolManager()->DisplayMessage(LOCTEXT("CopyChannelNotification", "UV Channel copied."), EToolMessageLevel::UserNotification);
 }
 
 void UUVEditorChannelEditTool::DeleteChannel()
@@ -697,6 +701,8 @@ void UUVEditorChannelEditTool::DeleteChannel()
 
 	SourceChannelProperties->Initialize(Targets, false);
 	SourceChannelProperties->TargetChannel = SourceChannelProperties->GetUVChannelNames()[Targets[ActiveAsset]->UVLayerIndex];
+
+	GetToolManager()->DisplayMessage(LOCTEXT("DeleteChannelNotification", "UV Channel deleted."), EToolMessageLevel::UserNotification);
 }
 
 void UUVEditorChannelEditTool::RecordAnalytics()
