UVEditor: Fix the way UV Editor detail customizations are unregistered, to fix the module sometimes crashing on unload.

#rnx

diff --git a/Engine/Plugins/Experimental/UVEditor/Source/UVEditor/Private/UVEditorModule.cpp b/Engine/Plugins/Experimental/UVEditor/Source/UVEditor/Private/UVEditorModule.cpp
--- a/Engine/Plugins/Experimental/UVEditor/Source/UVEditor/Private/UVEditorModule.cpp
+++ b/Engine/Plugins/Experimental/UVEditor/Source/UVEditor/Private/UVEditorModule.cpp
@@ -30,8 +30,10 @@ void FUVEditorModule::StartupModule()
 
 	// Register details view customizations
 	FPropertyEditorModule& PropertyModule = FModuleManager::LoadModuleChecked<FPropertyEditorModule>("PropertyEditor");
+	ClassesToUnregisterOnShutdown.Reset();
 	PropertyModule.RegisterCustomClassLayout(USelectToolActionPropertySet::StaticClass()->GetFName(), 
 		FOnGetDetailCustomizationInstance::CreateStatic(&FUVSelectToolActionPropertySetDetails::MakeInstance));
+	ClassesToUnregisterOnShutdown.Add(USelectToolActionPropertySet::StaticClass()->GetFName());
 }
 
 void FUVEditorModule::ShutdownModule()
@@ -46,7 +48,13 @@ void FUVEditorModule::ShutdownModule()
 
 	// Unregister customizations
 	FPropertyEditorModule* PropertyEditorModule = FModuleManager::GetModulePtr<FPropertyEditorModule>("PropertyEditor");
-	PropertyEditorModule->UnregisterCustomClassLayout(USelectToolActionPropertySet::StaticClass()->GetFName());
+	if (PropertyEditorModule)
+	{
+		for (FName ClassName : ClassesToUnregisterOnShutdown)
+		{
+			PropertyEditorModule->UnregisterCustomClassLayout(ClassName);
+		}
+	}
 }
 
 void FUVEditorModule::RegisterMenus()
diff --git a/Engine/Plugins/Experimental/UVEditor/Source/UVEditor/Public/UVEditorModule.h b/Engine/Plugins/Experimental/UVEditor/Source/UVEditor/Public/UVEditorModule.h
--- a/Engine/Plugins/Experimental/UVEditor/Source/UVEditor/Public/UVEditorModule.h
+++ b/Engine/Plugins/Experimental/UVEditor/Source/UVEditor/Public/UVEditorModule.h
@@ -22,9 +22,12 @@ public:
 	DECLARE_EVENT_OneParam(FUVEditorModule, FOnRegisterLayoutExtensions, FLayoutExtender&);
 	virtual FOnRegisterLayoutExtensions& OnRegisterLayoutExtensions() { return RegisterLayoutExtensions; }
 
+protected:
+	void RegisterMenus();
+
 private:
 	FOnRegisterLayoutExtensions	RegisterLayoutExtensions;
 
-protected:
-	void RegisterMenus();
+	/** StaticClass is not safe on shutdown, so we cache the name, and use this to unregister on shut down */
+	TArray<FName> ClassesToUnregisterOnShutdown;
 };
