UVEditor: Fix instances of UV editor breaking each others' channel selection menus by making GetToolkitFName give a unique name per UV editor instance.

#rnx

diff --git a/Engine/Plugins/Experimental/UVEditor/Source/UVEditor/Private/UVEditorToolkit.cpp b/Engine/Plugins/Experimental/UVEditor/Source/UVEditor/Private/UVEditorToolkit.cpp
--- a/Engine/Plugins/Experimental/UVEditor/Source/UVEditor/Private/UVEditorToolkit.cpp
+++ b/Engine/Plugins/Experimental/UVEditor/Source/UVEditor/Private/UVEditorToolkit.cpp
@@ -149,11 +149,17 @@ FText FUVEditorToolkit::GetToolkitName() const
 	return LOCTEXT("UVEditorMultipleTabName", "UV: Multiple");
 }
 
-// This gets used multiple places, most notably in GetToolMenuAppName, which gets
-// used to refer to menus/toolbars internally.
+// TODO: This may need changing once the editor team determines the proper course of action here.
+// In other asset editor toolkits, this would usually always return the same string, such
+// as "UVEditor". However, FAssetEditorToolkit::GetToolMenuToolbarName, which uses
+// FAssetEditorToolkit::GetToolMenuAppName (both are non-virtual) falls through to GetToolkitFName
+// for non-primary editors, rather than giving a unique name based on the edited UObject as it does
+// for primary editors that edit a single asset. We need to be able to customize the toolbar differently
+// for different instances of the UV editor to display different channel selection options, so we need 
+// the toolbar name to be instance-dependent, hence the unique name in GetToolkitFName here.
 FName FUVEditorToolkit::GetToolkitFName() const
 {
-	return FName("UVEditor");
+	return FName(FString::Printf(TEXT("UVEditor%p"), this));
 }
 
 // TODO: What is this actually used for?
