UV Editor: Fix checkboard background material to properly support multiple UV channels per individual asset.

#rnx

diff --git a/Engine/Plugins/Experimental/UVEditor/Source/UVEditor/Private/UVEditorMode.cpp b/Engine/Plugins/Experimental/UVEditor/Source/UVEditor/Private/UVEditorMode.cpp
--- a/Engine/Plugins/Experimental/UVEditor/Source/UVEditor/Private/UVEditorMode.cpp
+++ b/Engine/Plugins/Experimental/UVEditor/Source/UVEditor/Private/UVEditorMode.cpp
@@ -613,8 +613,12 @@ void UUVEditorMode::UpdatePreviewMaterialBasedOnBackground()
 	}
 	if (BackgroundVisualization->Settings->bVisible)
 	{
-		for (int32 AssetID = 0; AssetID < ToolInputObjects.Num(); ++AssetID) {
-			ToolInputObjects[AssetID]->AppliedPreview->OverrideMaterial = BackgroundVisualization->BackgroundMaterial;
+		for (int32 AssetID = 0; AssetID < ToolInputObjects.Num(); ++AssetID) {			
+			TObjectPtr<UMaterialInstanceDynamic> BackgroundMaterialOverride = UMaterialInstanceDynamic::Create(BackgroundVisualization->BackgroundMaterial->Parent, this);
+			BackgroundMaterialOverride->CopyInterpParameters(BackgroundVisualization->BackgroundMaterial);
+			BackgroundMaterialOverride->SetScalarParameterValue(TEXT("UVChannel"), GetDisplayedChannel(AssetID));
+
+			ToolInputObjects[AssetID]->AppliedPreview->OverrideMaterial = BackgroundMaterialOverride;
 		}
 	}
 }
@@ -741,6 +745,7 @@ void UUVEditorMode::ChangeInputObjectLayer(int32 AssetID, int32 NewLayerIndex)
 	}
 
 	PendingUVLayerIndex[AssetID] = InputObject->UVLayerIndex;
+	UpdatePreviewMaterialBasedOnBackground();
 }
 
 void UUVEditorMode::SetDisplayedUVChannels(const TArray<int32>& LayerPerAsset, bool bEmitUndoTransaction)
