MeshModelingTools: Move and rename UMeshSelectionMechanic, since it is currently UV Editor specific. Move and rename FDynamicMeshSelection since it is only used in UV Editor. Make the mechanic use FUVEditorUXSettings directly.

#rnx

diff --git a/Engine/Plugins/Experimental/UVEditor/Source/UVEditorTools/Private/Selection/UVEditorMeshSelectionMechanic.cpp b/Engine/Plugins/Experimental/UVEditor/Source/UVEditorTools/Private/Selection/UVEditorMeshSelectionMechanic.cpp
new file mode 100644
--- /dev/null
+++ b/Engine/Plugins/Experimental/UVEditor/Source/UVEditorTools/Private/Selection/UVEditorMeshSelectionMechanic.cpp
@@ -0,0 +1,1309 @@
+// Copyright Epic Games, Inc. All Rights Reserved.
+
+#include "Selection/UVEditorMeshSelectionMechanic.h"
+
+#include "BaseBehaviors/BehaviorTargetInterfaces.h"
+#include "BaseBehaviors/SingleClickOrDragBehavior.h"
+#include "BaseBehaviors/MouseHoverBehavior.h"
+#include "Drawing/TriangleSetComponent.h"
+#include "Drawing/LineSetComponent.h"
+#include "Drawing/PointSetComponent.h"
+#include "Drawing/PreviewGeometryActor.h"
+#include "Intersection/IntrTriangle2AxisAlignedBox2.h"
+#include "Intersection/IntersectionQueries2.h"
+#include "InteractiveToolManager.h"
+#include "Polyline3.h"
+#include "Selections/MeshConnectedComponents.h"
+#include "Spatial/GeometrySet3.h"
+#include "ToolSceneQueriesUtil.h"
+#include "ToolSetupUtil.h"
+#include "UVEditorUXSettings.h"
+
+using namespace UE::Geometry;
+
+#define LOCTEXT_NAMESPACE "UUVEditorMeshSelectionMechanic"
+
+namespace UVEditorMeshSelectionMechanicLocals
+{
+	/**
+	 * Class that can be used to apply/revert selection changes. Expects that the mesh selection
+	 * mechanic is the associated UObject that is passed to Apply/Revert.
+	 */
+	class FMeshSelectionMechanicSelectionChange : public FToolCommandChange
+	{
+	public:
+		FMeshSelectionMechanicSelectionChange(const FUVEditorDynamicMeshSelection& OldSelection,
+			const FUVEditorDynamicMeshSelection& NewSelection)
+			: Before(OldSelection)
+			, After(NewSelection)
+		{
+		}
+
+		virtual void Apply(UObject* Object) override
+		{
+			UUVEditorMeshSelectionMechanic* Mechanic = Cast<UUVEditorMeshSelectionMechanic>(Object);
+			if (Mechanic)
+			{
+				Mechanic->SetSelection(After, true, false);
+			}
+		}
+
+		virtual void Revert(UObject* Object) override
+		{
+			UUVEditorMeshSelectionMechanic* Mechanic = Cast<UUVEditorMeshSelectionMechanic>(Object);
+			if (Mechanic)
+			{
+				Mechanic->SetSelection(Before, true, false);
+			}
+		}
+
+		virtual FString ToString() const override
+		{
+			return TEXT("UVEditorMeshSelectionMechanicLocals::FMeshSelectionMechanicSelectionChange");
+		}
+
+	protected:
+		FUVEditorDynamicMeshSelection Before;
+		FUVEditorDynamicMeshSelection After;
+	};
+
+	template <typename InElementType>
+	void ToggleItem(TSet<InElementType>& Set, InElementType Item)
+	{
+		if (Set.Remove(Item) == 0)
+		{
+			Set.Add(Item);
+		}
+	}
+	
+	FUVEditorDynamicMeshSelection::EType ToCompatibleDynamicMeshSelectionType(const EUVEditorMeshSelectionMode& Mode)
+	{
+		switch (Mode)
+		{
+			case EUVEditorMeshSelectionMode::Mesh:
+			case EUVEditorMeshSelectionMode::Component:
+			case EUVEditorMeshSelectionMode::Triangle:
+				return FUVEditorDynamicMeshSelection::EType::Triangle;
+			case EUVEditorMeshSelectionMode::Edge:
+				return FUVEditorDynamicMeshSelection::EType::Edge;
+			case EUVEditorMeshSelectionMode::Vertex:
+				return FUVEditorDynamicMeshSelection::EType::Vertex;
+		}
+		checkNoEntry();
+		return FUVEditorDynamicMeshSelection::EType::Vertex;
+	}
+
+	// Returns the marquee selection rectangle, obtained from the given CameraRectangle, projected to the XY plane
+	FAxisAlignedBox2d GetRectangleXY(const FCameraRectangle& CameraRectangle)
+	{
+		ensure(CameraRectangle.bIsInitialized);
+		FAxisAlignedBox2d Result;
+		
+		double Offset = CameraRectangle.SelectionDomain.Plane.DistanceTo(FVector::ZeroVector);
+		FCameraRectangle::FRectangleInPlane Domain = CameraRectangle.ProjectSelectionDomain(Offset);
+		
+		// This works because we know the UV axes are aligned with the XY axes, see the comment in UUVEditorMode::InitializeTargets
+		const FVector MinPoint3D = CameraRectangle.PointUVToPoint3D(Domain.Plane, Domain.Rectangle.Min);
+		const FVector MaxPoint3D = CameraRectangle.PointUVToPoint3D(Domain.Plane, Domain.Rectangle.Max);
+		Result.Contain(FVector2d{MinPoint3D.X, MinPoint3D.Y}); // Convert to 2D and convert to double
+		Result.Contain(FVector2d{MaxPoint3D.X, MaxPoint3D.Y});
+	
+		return Result;
+	}
+	
+	FVector2d XY(const FVector3d& Point)
+	{
+		return {Point.X, Point.Y};
+	}
+
+	void AppendVertexIDs(const FDynamicMesh3& MeshXY0, int TriangleID, TArray<int>& VertexIDs)
+	{
+		const FIndex3i& Triangle = MeshXY0.GetTriangleRef(TriangleID);
+		VertexIDs.Add(Triangle.A);
+		VertexIDs.Add(Triangle.B);
+		VertexIDs.Add(Triangle.C);
+	}
+
+	void AppendVertexIDsIfIntersected(const FDynamicMesh3& MeshXY0, const FAxisAlignedBox2d& RectangleXY, int TriangleID, TArray<int>& VertexIDs)
+	{
+		const FIndex3i& Triangle = MeshXY0.GetTriangleRef(TriangleID);
+		if (RectangleXY.Contains(XY(MeshXY0.GetVertex(Triangle.A))))
+		{
+			VertexIDs.Add(Triangle.A);
+		}
+		
+		if (RectangleXY.Contains(XY(MeshXY0.GetVertex(Triangle.B))))
+		{
+			VertexIDs.Add(Triangle.B);
+		}
+		
+		if (RectangleXY.Contains(XY(MeshXY0.GetVertex(Triangle.C))))
+		{
+			VertexIDs.Add(Triangle.C);
+		}
+	}
+	
+	void AppendEdgeIDs(const FDynamicMesh3& MeshXY0, int TriangleID, TArray<int>& EdgeIDs)
+	{
+		const FIndex3i& Edges = MeshXY0.GetTriEdgesRef(TriangleID);
+		EdgeIDs.Add(Edges.A);
+		EdgeIDs.Add(Edges.B);
+		EdgeIDs.Add(Edges.C);
+	}
+	
+	void AppendEdgeIDsIfIntersected(const FDynamicMesh3& MeshXY0, const FAxisAlignedBox2d& RectangleXY, int TriangleID, TArray<int>& EdgeIDs)
+	{
+		const FIndex3i& Edges = MeshXY0.GetTriEdgesRef(TriangleID);
+
+		const FIndex2i& EdgeA = MeshXY0.GetEdgeRef(Edges.A).Vert;
+		const FSegment2d SegmentA(XY(MeshXY0.GetVertex(EdgeA.A)), XY(MeshXY0.GetVertex(EdgeA.B)));
+		if (TestIntersection(SegmentA, RectangleXY))
+		{
+			EdgeIDs.Add(Edges.A);
+		}
+		
+		const FIndex2i& EdgeB = MeshXY0.GetEdgeRef(Edges.B).Vert;
+		const FSegment2d SegmentB(XY(MeshXY0.GetVertex(EdgeB.A)), XY(MeshXY0.GetVertex(EdgeB.B)));
+		if (TestIntersection(SegmentB, RectangleXY))
+		{
+			EdgeIDs.Add(Edges.B);
+		}
+		
+		const FIndex2i& EdgeC = MeshXY0.GetEdgeRef(Edges.C).Vert;
+		const FSegment2d SegmentC(XY(MeshXY0.GetVertex(EdgeC.A)), XY(MeshXY0.GetVertex(EdgeC.B)));
+		if (TestIntersection(SegmentC, RectangleXY))
+		{
+			EdgeIDs.Add(Edges.C);
+		}
+	}
+	
+	void AppendTriangleID(const FDynamicMesh3&, int TriangleID, TArray<int>& TriangleIDs)
+	{
+		TriangleIDs.Add(TriangleID);
+	}
+
+	void AppendTriangleIDIfIntersected(const FDynamicMesh3& MeshXY0, const FAxisAlignedBox2d& RectangleXY, int TriangleID, TArray<int>& TriangleIDs)
+	{
+		const FIndex3i& Triangle = MeshXY0.GetTriangleRef(TriangleID);
+		const FTriangle2d TriangleXY(XY(MeshXY0.GetVertex(Triangle.A)),
+									 XY(MeshXY0.GetVertex(Triangle.B)),
+									 XY(MeshXY0.GetVertex(Triangle.C)));
+		
+		// Check with bTriangleIsOriented = false since some triangles maybe oriented away from the camera
+		if (FIntrTriangle2AxisAlignedBox2d Intersects(TriangleXY, RectangleXY, false); Intersects.Test())
+		{
+			TriangleIDs.Add(TriangleID);
+		}
+	}
+	
+	// Returns indices, collected by the given functions, from triangles which are intersected by the given rectangle.
+	// TreeXY0 must contain a mesh with vertices in the XY plane (have zero Z coordinate)
+	template<typename IDsFromTriangleF, typename IDsFromTriangleIfIntersectedF>
+	TArray<int32> FindAllIntersectionsAxisAlignedBox2(const FDynamicMeshAABBTree3& TreeXY0,
+													  const FAxisAlignedBox2d& RectangleXY,
+													  IDsFromTriangleF AppendIDs,
+													  IDsFromTriangleIfIntersectedF AppendIDsIfIntersected)
+	{
+		TRACE_CPUPROFILER_EVENT_SCOPE(FindAllIntersectionsAxisAlignedBox2);
+		
+		check(TreeXY0.GetMesh());
+		
+		TArray<int32> Result;
+		FAxisAlignedBox2d TreeRectangleXY;
+		TreeRectangleXY.Contain(XY(TreeXY0.GetBoundingBox().Min));
+		TreeRectangleXY.Contain(XY(TreeXY0.GetBoundingBox().Max));
+		if (RectangleXY.Contains(TreeRectangleXY))
+		{
+			// Early out selecting everything
+			Result.Reserve(TreeXY0.GetMesh()->TriangleCount());
+			for (int TriangleID : TreeXY0.GetMesh()->TriangleIndicesItr())
+			{
+				AppendIDs(*TreeXY0.GetMesh(), TriangleID, Result);
+			}
+			return Result;
+		}
+		
+		int SelectAllDepth = TNumericLimits<int>::Max();
+		int CurrentDepth = -1;
+		
+		// Traversal is depth first
+		FDynamicMeshAABBTree3::FTreeTraversal Traversal;
+		
+		Traversal.NextBoxF =
+			[&RectangleXY, &SelectAllDepth, &CurrentDepth](const FAxisAlignedBox3d& Box, int Depth)
+		{
+			CurrentDepth = Depth;
+			if (Depth > SelectAllDepth)
+			{
+				// We are deeper than the depth whose AABB was first detected to be contained in the RectangleXY,
+				// descend and collect all leaf triangles
+				return true;
+			}
+			
+			SelectAllDepth = TNumericLimits<int>::Max();
+			
+			const FAxisAlignedBox2d BoxXY(XY(Box.Min), XY(Box.Max));
+			if (RectangleXY.Intersects(BoxXY))
+			{
+				if (RectangleXY.Contains(BoxXY))
+				{
+					SelectAllDepth = Depth;
+				}
+				
+				return true;		
+			}
+			return false;
+		};
+		
+		Traversal.NextTriangleF =
+			[&RectangleXY, &SelectAllDepth, &CurrentDepth, &TreeXY0, &Result, &AppendIDs, &AppendIDsIfIntersected]
+			(int TriangleID)
+		{
+			if (CurrentDepth >= SelectAllDepth)
+			{
+				// This TriangleID is entirely contained in the selection rectangle so we can skip intersection testing
+				return AppendIDs(*TreeXY0.GetMesh(), TriangleID, Result);
+			}
+			return AppendIDsIfIntersected(*TreeXY0.GetMesh(), RectangleXY, TriangleID, Result);
+		};
+		
+		TreeXY0.DoTraversal(Traversal);
+
+		return Result;
+	}
+} // namespace UVEditorMeshSelectionMechanicLocals
+
+void UUVEditorMeshSelectionMechanic::Setup(UInteractiveTool* ParentToolIn)
+{
+	UInteractionMechanic::Setup(ParentToolIn);
+
+	// This will be the target for the click drag behavior below
+	MarqueeMechanic = NewObject<URectangleMarqueeMechanic>();
+	MarqueeMechanic->bUseExternalClickDragBehavior = true;
+	MarqueeMechanic->Setup(ParentToolIn);
+	MarqueeMechanic->OnDragRectangleStarted.AddUObject(this, &UUVEditorMeshSelectionMechanic::OnDragRectangleStarted);
+	// TODO(Performance) :DynamicMarqueeSelection It would be cool to have the marquee selection update dynamically as
+	//  the rectangle gets changed, right now this isn't interactive for large meshes so we disabled it
+	// MarqueeMechanic->OnDragRectangleChanged.AddUObject(this, &UUVEditorMeshSelectionMechanic::OnDragRectangleChanged);
+	MarqueeMechanic->OnDragRectangleFinished.AddUObject(this, &UUVEditorMeshSelectionMechanic::OnDragRectangleFinished);
+
+	USingleClickOrDragInputBehavior* ClickOrDragBehavior = NewObject<USingleClickOrDragInputBehavior>();
+	ClickOrDragBehavior->Initialize(this, MarqueeMechanic);
+	ClickOrDragBehavior->Modifiers.RegisterModifier(ShiftModifierID, FInputDeviceState::IsShiftKeyDown);
+	ClickOrDragBehavior->Modifiers.RegisterModifier(CtrlModifierID, FInputDeviceState::IsCtrlKeyDown);
+	ParentTool->AddInputBehavior(ClickOrDragBehavior);
+	
+	Settings = NewObject<UUVEditorMeshSelectionMechanicProperties>(ParentToolIn);
+	Settings->RestoreProperties(ParentToolIn);
+	AddToolPropertySource(Settings);
+
+	ULocalMouseHoverBehavior* HoverBehavior = NewObject<ULocalMouseHoverBehavior>();
+	HoverBehavior->Initialize();
+	HoverBehavior->BeginHitTestFunc = [this](const FInputDeviceRay& Pos)
+	{
+		FInputRayHit Hit;
+		if (Settings->bShowHoveredElements)
+		{
+			EUVEditorMeshSelectionMode Mode = SelectionMode;
+			if (Mode != EUVEditorMeshSelectionMode::Vertex && Mode != EUVEditorMeshSelectionMode::Edge)
+			{
+				Mode = EUVEditorMeshSelectionMode::Triangle;
+			}
+
+			// We don't bother with the depth since there aren't competing hover behaviors
+			Hit.bHit = RayCast(Pos, Mode).IsEmpty() == false;
+		}
+		else
+		{
+			// Disable hover behavior
+			Hit.bHit = false;
+		}
+		return Hit;
+	};
+	HoverBehavior->OnUpdateHoverFunc = [this](const FInputDeviceRay& Pos)
+	{
+		EUVEditorMeshSelectionMode Mode = SelectionMode;
+		if (SelectionMode != EUVEditorMeshSelectionMode::Vertex && SelectionMode != EUVEditorMeshSelectionMode::Edge)
+		{
+			Mode = EUVEditorMeshSelectionMode::Triangle;
+		}
+		
+		HoverPointSet->Clear();
+		HoverLineSet->Clear();
+		HoverTriangleSet->Clear();
+		
+		TSet<int32> Hovered = RayCast(Pos, Mode);
+		for (int32 Id : Hovered)
+		{
+			if (SelectionMode == EUVEditorMeshSelectionMode::Vertex)
+			{
+				const FVector& P = CurrentSelection.Mesh->GetVertexRef(Id);
+				const FRenderablePoint PointToRender(P, 
+					FUVEditorUXSettings::SelectionHoverTriangleWireframeColor, 
+					FUVEditorUXSettings::SelectionPointThickness);
+				
+				HoverPointSet->AddPoint(PointToRender);
+			}
+			else if (SelectionMode == EUVEditorMeshSelectionMode::Edge)
+			{
+				const FIndex2i EdgeVids = CurrentSelection.Mesh->GetEdgeV(Id);
+				const FVector& A = CurrentSelection.Mesh->GetVertexRef(EdgeVids.A);
+				const FVector& B = CurrentSelection.Mesh->GetVertexRef(EdgeVids.B);
+				
+				HoverLineSet->AddLine(A, B, 
+					FUVEditorUXSettings::SelectionHoverTriangleWireframeColor, 
+					FUVEditorUXSettings::SelectionLineThickness,
+					FUVEditorUXSettings::SelectionHoverWireframeDepthBias);
+			}
+			else
+			{
+				const FIndex3i Vids = CurrentSelection.Mesh->GetTriangle(Id);
+				const FVector& A = CurrentSelection.Mesh->GetVertex(Vids[0]);
+				const FVector& B = CurrentSelection.Mesh->GetVertex(Vids[1]);
+				const FVector& C = CurrentSelection.Mesh->GetVertex(Vids[2]);
+
+				HoverLineSet->AddLine(A, B, FUVEditorUXSettings::SelectionHoverTriangleWireframeColor, 
+					FUVEditorUXSettings::SelectionLineThickness, FUVEditorUXSettings::SelectionHoverWireframeDepthBias);
+				HoverLineSet->AddLine(B, C, FUVEditorUXSettings::SelectionHoverTriangleWireframeColor, 
+					FUVEditorUXSettings::SelectionLineThickness, FUVEditorUXSettings::SelectionHoverWireframeDepthBias);
+				HoverLineSet->AddLine(C, A, FUVEditorUXSettings::SelectionHoverTriangleWireframeColor, 
+					FUVEditorUXSettings::SelectionLineThickness, FUVEditorUXSettings::SelectionHoverWireframeDepthBias);
+				HoverTriangleSet->AddTriangle(A, B, C, FVector::ZAxisVector, 
+					FUVEditorUXSettings::SelectionHoverTriangleFillColor, HoverTriangleSetMaterial);
+			}
+
+			return true;
+		}
+		return false;
+	};
+	HoverBehavior->OnEndHoverFunc = [this]()
+	{
+		HoverPointSet->Clear();
+		HoverLineSet->Clear();
+		HoverTriangleSet->Clear();
+	};
+	ParentTool->AddInputBehavior(HoverBehavior);
+
+	ClearCurrentSelection();
+
+	// Setup selection visualizations
+	TriangleSet = NewObject<UTriangleSetComponent>();
+	// We are setting the TranslucencySortPriority here to handle the UV editor's use case in 2D
+	// where multiple translucent layers are drawn on top of each other but still need depth sorting.
+	TriangleSet->TranslucencySortPriority = FUVEditorUXSettings::SelectionTriangleDepthBias;	
+	TriangleSetMaterial = ToolSetupUtil::GetCustomTwoSidedDepthOffsetMaterial(GetParentTool()->GetToolManager(),
+		FUVEditorUXSettings::SelectionTriangleFillColor, 
+		FUVEditorUXSettings::SelectionTriangleDepthBias, 
+		FUVEditorUXSettings::SelectionTriangleOpacity);
+	
+	PointSet = NewObject<UPointSetComponent>();
+	PointSet->SetPointMaterial(ToolSetupUtil::GetDefaultPointComponentMaterial( GetParentTool()->GetToolManager(), true));
+	LineSet = NewObject<ULineSetComponent>();
+	LineSet->SetLineMaterial(ToolSetupUtil::GetDefaultLineComponentMaterial( GetParentTool()->GetToolManager(), true));
+
+	// Setup hover visualizations
+	HoverPointSet = NewObject<UPointSetComponent>();
+	HoverPointSet->SetPointMaterial(ToolSetupUtil::GetDefaultPointComponentMaterial(GetParentTool()->GetToolManager(), false));
+	HoverLineSet = NewObject<ULineSetComponent>();
+	HoverLineSet->SetLineMaterial(ToolSetupUtil::GetDefaultLineComponentMaterial(GetParentTool()->GetToolManager(), false));
+	HoverTriangleSet = NewObject<UTriangleSetComponent>();
+	HoverTriangleSetMaterial = ToolSetupUtil::GetCustomTwoSidedDepthOffsetMaterial(GetParentTool()->GetToolManager(),
+		FUVEditorUXSettings::SelectionHoverTriangleFillColor, 
+		FUVEditorUXSettings::SelectionHoverTriangleDepthBias, 
+		FUVEditorUXSettings::SelectionHoverTriangleOpacity);
+
+	// Add default selection change emitter if one was not provided.
+	if (!EmitSelectionChange)
+	{
+		EmitSelectionChange = [this](const FUVEditorDynamicMeshSelection& OldSelection, const FUVEditorDynamicMeshSelection& NewSelection)
+		{
+			GetParentTool()->GetToolManager()->EmitObjectChange(this, 
+				MakeUnique<UVEditorMeshSelectionMechanicLocals::FMeshSelectionMechanicSelectionChange>(OldSelection, NewSelection),
+				LOCTEXT("SelectionChangeMessage", "Selection Change"));
+		};
+	}
+}
+
+void UUVEditorMeshSelectionMechanic::SetWorld(UWorld* World)
+{
+	auto RegisterComponent = [this](UMeshComponent* Component)
+	{
+		if (Component->IsRegistered())
+		{
+			Component->ReregisterComponent();
+		}
+		else
+		{
+			Component->RegisterComponent();
+		}
+	};
+	
+	// It may be unreasonable to worry about SetWorld being called more than once, but let's be safe anyway
+	if (PreviewGeometryActor)
+	{
+		PreviewGeometryActor->Destroy();
+	}
+
+	// We need the world so we can create the geometry actor in the right place.
+	FRotator Rotation(0.0f, 0.0f, 0.0f);
+	FActorSpawnParameters SpawnInfo;
+	PreviewGeometryActor = World->SpawnActor<APreviewGeometryActor>(FVector::ZeroVector, Rotation, SpawnInfo);
+
+
+	// Attach the rendering component to the actor		
+	TriangleSet->Rename(nullptr, PreviewGeometryActor); // Changes the "outer"
+	PreviewGeometryActor->SetRootComponent(TriangleSet);
+	RegisterComponent(TriangleSet);
+
+	LineSet->Rename(nullptr, PreviewGeometryActor); // Changes the "outer"
+	LineSet->AttachToComponent(TriangleSet, FAttachmentTransformRules::KeepWorldTransform);
+	RegisterComponent(LineSet);
+	
+	PointSet->Rename(nullptr, PreviewGeometryActor); // Changes the "outer"
+	PointSet->AttachToComponent(TriangleSet, FAttachmentTransformRules::KeepWorldTransform);
+	RegisterComponent(PointSet);
+
+	if (HoverGeometryActor)
+	{
+		HoverGeometryActor->Destroy();
+	}
+	
+	HoverGeometryActor = World->SpawnActor<APreviewGeometryActor>(FVector::ZeroVector, Rotation, SpawnInfo);
+	
+	// Attach the rendering component to the actor		
+	HoverTriangleSet->Rename(nullptr, HoverGeometryActor); // Changes the "outer"
+	HoverGeometryActor->SetRootComponent(HoverTriangleSet);
+	RegisterComponent(HoverTriangleSet);
+
+	HoverLineSet->Rename(nullptr, HoverGeometryActor); // Changes the "outer"
+	HoverLineSet->AttachToComponent(HoverTriangleSet, FAttachmentTransformRules::KeepWorldTransform);
+	RegisterComponent(HoverLineSet);
+	
+	HoverPointSet->Rename(nullptr, HoverGeometryActor); // Changes the "outer"
+	HoverPointSet->AttachToComponent(HoverTriangleSet, FAttachmentTransformRules::KeepWorldTransform);
+	RegisterComponent(HoverPointSet);
+}
+
+void UUVEditorMeshSelectionMechanic::Shutdown()
+{
+	if (PreviewGeometryActor)
+	{
+		PreviewGeometryActor->Destroy();
+		PreviewGeometryActor = nullptr;
+	}
+}
+
+void UUVEditorMeshSelectionMechanic::AddSpatial(TSharedPtr<FDynamicMeshAABBTree3> SpatialIn, const FTransform& TransformIn)
+{
+	MeshSpatials.Add(SpatialIn);
+	MeshTransforms.Add(TransformIn);
+}
+
+const FUVEditorDynamicMeshSelection& UUVEditorMeshSelectionMechanic::GetCurrentSelection() const
+{
+	return CurrentSelection;
+}
+
+void UUVEditorMeshSelectionMechanic::SetSelection(const FUVEditorDynamicMeshSelection& Selection, bool bBroadcast, bool bEmitChange)
+{
+	FUVEditorDynamicMeshSelection OriginalSelection = CurrentSelection;
+	CurrentSelection = Selection;
+
+	// Adjust the mesh in which the selection lives
+	if (CurrentSelection.Mesh == nullptr)
+	{
+		// We're actually clearing the selection
+		if (!ensure(CurrentSelection.SelectedIDs.IsEmpty()))
+		{
+			CurrentSelection.SelectedIDs.Empty();
+		}
+		CurrentSelectionIndex = IndexConstants::InvalidID;
+	}
+	else if (CurrentSelectionIndex == IndexConstants::InvalidID 
+		|| MeshSpatials[CurrentSelectionIndex]->GetMesh() != Selection.Mesh)
+	{
+		CurrentSelectionIndex = IndexConstants::InvalidID;
+		for (int32 i = 0; i < MeshSpatials.Num(); ++i)
+		{
+			if (MeshSpatials[i]->GetMesh() == Selection.Mesh)
+			{
+				CurrentSelectionIndex = i;
+				break;
+			}
+		}
+		ensure(CurrentSelectionIndex != IndexConstants::InvalidID);
+	}
+
+	UpdateCentroid();
+
+	if (bEmitChange && OriginalSelection != CurrentSelection)
+	{
+		EmitSelectionChange(OriginalSelection, CurrentSelection);
+	}
+	if (bBroadcast)
+	{
+		OnSelectionChanged.Broadcast();
+	}
+
+	// Rebuild after broadcast in case the outside world wants to adjust things like color...
+	RebuildDrawnElements(FTransform(CurrentSelectionCentroid));
+}
+
+void UUVEditorMeshSelectionMechanic::RebuildDrawnElements(const FTransform& StartTransform)
+{
+	TRACE_CPUPROFILER_EVENT_SCOPE(MeshSelectionMechanic_RebuildDrawnElements);
+
+	TriangleSet->Clear();
+	LineSet->Clear();
+	PointSet->Clear();
+	PreviewGeometryActor->SetActorTransform(StartTransform);
+
+	if (CurrentSelectionIndex == IndexConstants::InvalidID)
+	{
+		return;
+	}
+
+	// For us to end up with the StartTransform, we have to invert it, but only after applying
+	// the mesh transform to begin with.
+	auto TransformToApply = [&StartTransform, this](const FVector& VectorIn)
+	{
+		return StartTransform.InverseTransformPosition(
+			MeshTransforms[CurrentSelectionIndex].TransformPosition(VectorIn));
+	};
+
+	if (CurrentSelection.Type == FUVEditorDynamicMeshSelection::EType::Triangle)
+	{
+		TRACE_CPUPROFILER_EVENT_SCOPE(MeshSelectionMechanic_RebuildDrawnElements_Triangle);
+		
+		TriangleSet->ReserveTriangles(CurrentSelection.SelectedIDs.Num());
+		LineSet->ReserveLines(CurrentSelection.SelectedIDs.Num() * 3);
+		for (int32 Tid : CurrentSelection.SelectedIDs)
+		{
+			FIndex3i Vids = CurrentSelection.Mesh->GetTriangle(Tid);
+			FVector Points[3];
+			for (int i = 0; i < 3; ++i)
+			{
+				Points[i] = TransformToApply(CurrentSelection.Mesh->GetVertex(Vids[i]));
+			}
+			TriangleSet->AddTriangle(Points[0], Points[1], Points[2], FVector(0, 0, 1), FUVEditorUXSettings::SelectionTriangleFillColor, TriangleSetMaterial);
+			for (int i = 0; i < 3; ++i)
+			{
+				int NextIndex = (i + 1) % 3;
+				LineSet->AddLine(Points[i], Points[NextIndex],
+					FUVEditorUXSettings::SelectionTriangleWireframeColor, 
+					FUVEditorUXSettings::SelectionLineThickness, 
+					FUVEditorUXSettings::SelectionWireframeDepthBias);
+			}
+		}
+	}
+	else if (CurrentSelection.Type == FUVEditorDynamicMeshSelection::EType::Edge)
+	{
+		TRACE_CPUPROFILER_EVENT_SCOPE(MeshSelectionMechanic_RebuildDrawnElements_Edge);
+
+		LineSet->ReserveLines(CurrentSelection.SelectedIDs.Num());
+		for (int32 Eid : CurrentSelection.SelectedIDs)
+		{
+			FIndex2i EdgeVids = CurrentSelection.Mesh->GetEdgeV(Eid);
+			LineSet->AddLine(
+				TransformToApply(CurrentSelection.Mesh->GetVertex(EdgeVids.A)),
+				TransformToApply(CurrentSelection.Mesh->GetVertex(EdgeVids.B)),
+				FUVEditorUXSettings::SelectionTriangleWireframeColor, 
+				FUVEditorUXSettings::SelectionLineThickness, 
+				FUVEditorUXSettings::SelectionWireframeDepthBias);
+		}
+	}
+	else if (CurrentSelection.Type == FUVEditorDynamicMeshSelection::EType::Vertex)
+	{
+		TRACE_CPUPROFILER_EVENT_SCOPE(MeshSelectionMechanic_RebuildDrawnElements_Vertex);
+
+		PointSet->ReservePoints(CurrentSelection.SelectedIDs.Num());
+		for (int32 Vid : CurrentSelection.SelectedIDs)
+		{
+			FRenderablePoint PointToRender(TransformToApply(CurrentSelection.Mesh->GetVertex(Vid)),
+				FUVEditorUXSettings::SelectionTriangleWireframeColor,
+				FUVEditorUXSettings::SelectionPointThickness,
+				FUVEditorUXSettings::SelectionWireframeDepthBias);
+			PointSet->AddPoint(PointToRender);
+		}
+	}
+}
+
+void UUVEditorMeshSelectionMechanic::UpdateCentroid()
+{
+	TRACE_CPUPROFILER_EVENT_SCOPE(MeshSelectionMechanic_UpdateCentroid);
+
+	CurrentSelectionCentroid = FVector3d(0);
+	if (CurrentSelection.IsEmpty())
+	{
+		return;
+	}
+
+	if (CurrentSelection.Type == FUVEditorDynamicMeshSelection::EType::Edge)
+	{
+		for (int32 Eid : CurrentSelection.SelectedIDs)
+		{
+			CurrentSelectionCentroid += CurrentSelection.Mesh->GetEdgePoint(Eid, 0.5);
+		}
+		CurrentSelectionCentroid /= CurrentSelection.SelectedIDs.Num();
+	}
+	else if (CurrentSelection.Type == FUVEditorDynamicMeshSelection::EType::Triangle)
+	{
+		for (int32 Tid : CurrentSelection.SelectedIDs)
+		{
+			CurrentSelectionCentroid += CurrentSelection.Mesh->GetTriCentroid(Tid);
+		}
+		CurrentSelectionCentroid /= CurrentSelection.SelectedIDs.Num();
+	}
+	if (CurrentSelection.Type == FUVEditorDynamicMeshSelection::EType::Vertex)
+	{
+		for (int32 Vid : CurrentSelection.SelectedIDs)
+		{
+			CurrentSelectionCentroid += CurrentSelection.Mesh->GetVertex(Vid);
+		}
+		CurrentSelectionCentroid /= CurrentSelection.SelectedIDs.Num();
+	}
+
+	// disable centroid caching if mesh does not have shape changestamp enabled
+	CentroidTimestamp = CurrentSelection.Mesh->HasShapeChangeStampEnabled() ?
+		CurrentSelection.Mesh->GetShapeChangeStamp() : TNumericLimits<uint32>::Max();
+}
+
+FVector3d UUVEditorMeshSelectionMechanic::GetCurrentSelectionCentroid()
+{
+	if (!CurrentSelection.Mesh)
+	{
+		return FVector3d(0);
+	}
+	else if (CurrentSelection.Mesh->GetShapeChangeStamp() != CentroidTimestamp)
+	{
+		UpdateCentroid();
+	}
+
+	return CurrentSelectionCentroid;
+}
+
+
+void UUVEditorMeshSelectionMechanic::SetDrawnElementsTransform(const FTransform& Transform)
+{
+	PreviewGeometryActor->SetActorTransform(Transform);
+}
+
+void UUVEditorMeshSelectionMechanic::Render(IToolsContextRenderAPI* RenderAPI)
+{
+	// Cache the camera state
+	MarqueeMechanic->Render(RenderAPI);
+	GetParentTool()->GetToolManager()->GetContextQueriesAPI()->GetCurrentViewState(CameraState);
+}
+
+void UUVEditorMeshSelectionMechanic::DrawHUD(FCanvas* Canvas, IToolsContextRenderAPI* RenderAPI)
+{
+	MarqueeMechanic->DrawHUD(Canvas, RenderAPI);
+}
+
+FInputRayHit UUVEditorMeshSelectionMechanic::IsHitByClick(const FInputDeviceRay& ClickPos)
+{
+	// TODO We could account for what modifiers would do and return an accurate depth here but for now this simple code works fine
+
+	// Return a hit so we always capture and can clear the selection
+	FInputRayHit Hit; 
+	Hit.bHit = true;
+	return Hit;	
+}
+
+void UUVEditorMeshSelectionMechanic::ClearCurrentSelection()
+{
+	CurrentSelection.SelectedIDs.Empty();
+	CurrentSelection.Mesh = nullptr;
+	CurrentSelection.Type = UVEditorMeshSelectionMechanicLocals::ToCompatibleDynamicMeshSelectionType(SelectionMode);
+	CurrentSelectionIndex = IndexConstants::InvalidID;
+}
+
+
+void UUVEditorMeshSelectionMechanic::ChangeSelectionMode(const EUVEditorMeshSelectionMode& TargetMode)
+{
+	TRACE_CPUPROFILER_EVENT_SCOPE(MeshSelectionMechanic_ChangeSelectionMode);
+	
+	if (CurrentSelection.SelectedIDs.IsEmpty())
+	{
+		SelectionMode = TargetMode;
+		return;
+	}
+	
+	const FUVEditorDynamicMeshSelection OriginalSelection = CurrentSelection;
+	
+	auto VerticesToEdges = [this, &OriginalSelection]()
+	{
+		TRACE_CPUPROFILER_EVENT_SCOPE(MeshSelectionMechanic_ChangeSelectionMode_VerticesToEdges);
+		
+		CurrentSelection.SelectedIDs.Empty();
+		CurrentSelection.Type = FUVEditorDynamicMeshSelection::EType::Edge;
+
+		for (int32 Vid : OriginalSelection.SelectedIDs)
+		{
+			for (int32 Eid : OriginalSelection.Mesh->VtxEdgesItr(Vid))
+			{
+				if (!CurrentSelection.SelectedIDs.Contains(Eid))
+				{
+					FIndex2i Verts = OriginalSelection.Mesh->GetEdgeV(Eid);
+					if (OriginalSelection.SelectedIDs.Contains(Verts.A) &&
+						OriginalSelection.SelectedIDs.Contains(Verts.B))
+					{
+						CurrentSelection.SelectedIDs.Add(Eid);	
+					}
+				}
+			}
+		}
+	};
+
+	auto VerticesToTriangles = [this, &OriginalSelection]()
+	{
+		TRACE_CPUPROFILER_EVENT_SCOPE(MeshSelectionMechanic_ChangeSelectionMode_VerticesToTriangles);
+		
+		CurrentSelection.SelectedIDs.Empty();
+		CurrentSelection.Type = FUVEditorDynamicMeshSelection::EType::Triangle;
+		
+		for (int32 Vid : OriginalSelection.SelectedIDs)
+		{
+			for (int32 Tid : OriginalSelection.Mesh->VtxTrianglesItr(Vid))
+			{
+				if (!CurrentSelection.SelectedIDs.Contains(Tid))
+				{
+					FIndex3i Verts = OriginalSelection.Mesh->GetTriangle(Tid);
+					if (OriginalSelection.SelectedIDs.Contains(Verts.A) &&
+						OriginalSelection.SelectedIDs.Contains(Verts.B) &&
+						OriginalSelection.SelectedIDs.Contains(Verts.C))
+					{
+						CurrentSelection.SelectedIDs.Add(Tid);	
+					}
+				}
+			}
+		}
+	};
+
+	auto EdgesToVertices = [this, &OriginalSelection]()
+	{
+		TRACE_CPUPROFILER_EVENT_SCOPE(MeshSelectionMechanic_ChangeSelectionMode_EdgesToVertices);
+		
+		CurrentSelection.SelectedIDs.Empty();
+		CurrentSelection.Type = FUVEditorDynamicMeshSelection::EType::Vertex;
+		
+		for (int32 Eid : OriginalSelection.SelectedIDs)
+		{
+			FIndex2i Verts = OriginalSelection.Mesh->GetEdgeV(Eid);
+			CurrentSelection.SelectedIDs.Add(Verts.A);
+			CurrentSelection.SelectedIDs.Add(Verts.B);
+		}
+	};
+
+	// Triangles with two selected edges will be selected
+	auto EdgesToTriangles = [this, &OriginalSelection]()
+	{
+		TRACE_CPUPROFILER_EVENT_SCOPE(MeshSelectionMechanic_ChangeSelectionMode_EdgesToTriangles);
+		
+		CurrentSelection.SelectedIDs.Empty();
+		CurrentSelection.Type = FUVEditorDynamicMeshSelection::EType::Triangle;
+		
+		TArray<int32> FoundTriangles;
+		for (int32 Eid : OriginalSelection.SelectedIDs)
+		{
+			FIndex2i Tris = OriginalSelection.Mesh->GetEdgeT(Eid);
+			FoundTriangles.Add(Tris.A);
+			if (Tris.B != IndexConstants::InvalidID)
+			{
+				FoundTriangles.Add(Tris.B);
+			}
+		}
+
+		if (FoundTriangles.Num() < 2)
+		{
+			return;
+		}
+
+		Algo::Sort(FoundTriangles);
+
+		for (int I = 0; I < FoundTriangles.Num() - 1; I++)
+		{
+			if (FoundTriangles[I] == FoundTriangles[I + 1])
+			{
+				CurrentSelection.SelectedIDs.Add(FoundTriangles[I]);
+				I++;
+			}	
+		}
+	};
+
+	auto TrianglesToVertices = [this, &OriginalSelection]()
+	{
+		TRACE_CPUPROFILER_EVENT_SCOPE(MeshSelectionMechanic_ChangeSelectionMode_TrianglesToVertices);
+		
+		CurrentSelection.SelectedIDs.Empty();
+		CurrentSelection.Type = FUVEditorDynamicMeshSelection::EType::Vertex;
+
+		for (int32 Tid : OriginalSelection.SelectedIDs)
+		{
+			FIndex3i Verts = OriginalSelection.Mesh->GetTriangle(Tid);
+			CurrentSelection.SelectedIDs.Add(Verts.A);
+			CurrentSelection.SelectedIDs.Add(Verts.B);
+			CurrentSelection.SelectedIDs.Add(Verts.C);
+		}
+	};
+
+	auto TrianglesToEdges = [this, &OriginalSelection]()
+	{
+		TRACE_CPUPROFILER_EVENT_SCOPE(MeshSelectionMechanic_ChangeSelectionMode_TrianglesToEdges);
+			
+		CurrentSelection.SelectedIDs.Empty();
+		CurrentSelection.Type = FUVEditorDynamicMeshSelection::EType::Edge;
+	
+		for (int32 Tid : OriginalSelection.SelectedIDs)
+		{
+			FIndex3i Edges = OriginalSelection.Mesh->GetTriEdgesRef(Tid);
+			CurrentSelection.SelectedIDs.Add(Edges.A);
+			CurrentSelection.SelectedIDs.Add(Edges.B);
+			CurrentSelection.SelectedIDs.Add(Edges.C);
+		}
+	};
+
+	switch (SelectionMode)
+	{
+	case EUVEditorMeshSelectionMode::Vertex:
+		switch (TargetMode)
+		{
+		case EUVEditorMeshSelectionMode::Vertex:
+			// Do nothing
+			break;
+
+		case EUVEditorMeshSelectionMode::Edge:
+			VerticesToEdges();
+			break;
+			
+		case EUVEditorMeshSelectionMode::Triangle:
+		case EUVEditorMeshSelectionMode::Component:
+		case EUVEditorMeshSelectionMode::Mesh:
+			VerticesToTriangles();
+			break;
+		}
+		break; // SelectionMode
+	
+	case EUVEditorMeshSelectionMode::Edge:
+		switch (TargetMode)
+		{
+		case EUVEditorMeshSelectionMode::Vertex:
+			EdgesToVertices();
+			break;
+
+		case EUVEditorMeshSelectionMode::Edge:
+			// Do nothing
+			break;
+
+		case EUVEditorMeshSelectionMode::Mesh:
+		case EUVEditorMeshSelectionMode::Triangle:
+		case EUVEditorMeshSelectionMode::Component:
+			EdgesToTriangles();
+			break;
+		}
+		break; // SelectionMode
+	
+	case EUVEditorMeshSelectionMode::Mesh:
+	case EUVEditorMeshSelectionMode::Triangle:
+	case EUVEditorMeshSelectionMode::Component:
+		switch (TargetMode)
+		{
+		case EUVEditorMeshSelectionMode::Vertex:
+			TrianglesToVertices();
+			break;
+
+		case EUVEditorMeshSelectionMode::Edge:
+			TrianglesToEdges();
+			break;
+
+		case EUVEditorMeshSelectionMode::Mesh:
+		case EUVEditorMeshSelectionMode::Triangle:
+		case EUVEditorMeshSelectionMode::Component:
+			// Do nothing
+			break;
+		}
+		break; // SelectionMode
+	}
+
+	if (CurrentSelection != OriginalSelection)
+	{
+		UpdateCentroid();
+		RebuildDrawnElements(FTransform(GetCurrentSelectionCentroid()));
+		EmitSelectionChange(OriginalSelection, CurrentSelection);
+		OnSelectionChanged.Broadcast();
+	}
+
+	SelectionMode = TargetMode;
+	return;
+}
+
+
+void UUVEditorMeshSelectionMechanic::UpdateCurrentSelection(const TSet<int32>& NewSelection, bool CalledFromOnDragRectangleChanged)
+{
+	TRACE_CPUPROFILER_EVENT_SCOPE(MeshSelectionMechanic_UpdateCurrentSelection);
+	
+	if (!ShouldRestartSelection() && CalledFromOnDragRectangleChanged)
+	{
+		// If we're modifying (adding/removing/toggling, but not restarting) the selection, we should start with the
+		// selection cached in OnDragRectangleStarted, otherwise multiple changes get accumulated as the rectangle is
+		// swept around
+		CurrentSelection.SelectedIDs = PreDragSelection.SelectedIDs;
+	}
+	
+	if (ShouldRestartSelection())
+	{
+		CurrentSelection.SelectedIDs = NewSelection;
+	}
+	else if (ShouldAddToSelection())
+	{
+		CurrentSelection.SelectedIDs.Append(NewSelection);
+	}
+	else if (ShouldToggleFromSelection())
+	{
+		for (int32 Index : NewSelection)
+		{
+			UVEditorMeshSelectionMechanicLocals::ToggleItem(CurrentSelection.SelectedIDs, Index);
+		}
+	}
+	else if (ShouldRemoveFromSelection())
+	{
+		CurrentSelection.SelectedIDs = CurrentSelection.SelectedIDs.Difference(NewSelection);
+	}
+	else
+	{
+		checkNoEntry();
+	}
+
+	if (CurrentSelection.SelectedIDs.IsEmpty())
+	{
+		ClearCurrentSelection();
+	}
+	
+	CurrentSelection.Type = UVEditorMeshSelectionMechanicLocals::ToCompatibleDynamicMeshSelectionType(SelectionMode);
+}
+
+void UUVEditorMeshSelectionMechanic::OnClicked(const FInputDeviceRay& ClickPos)
+{
+	TRACE_CPUPROFILER_EVENT_SCOPE(MeshSelectionMechanic_OnClicked);
+	
+	FUVEditorDynamicMeshSelection OriginalSelection = CurrentSelection;
+	
+	if (CurrentSelection.Type != UVEditorMeshSelectionMechanicLocals::ToCompatibleDynamicMeshSelectionType(SelectionMode))
+	{
+		ClearCurrentSelection();
+	}
+
+	TSet<int32> ClickSelectedIDs = RayCast(ClickPos, SelectionMode); // TODO Maybe using a TVariant would have been better here...
+
+	// TODO Perhaps selection clearing should happen only if the click occurs further than some threshold from the meshes
+	UpdateCurrentSelection(ClickSelectedIDs);
+
+	if (OriginalSelection != CurrentSelection)
+	{
+		UpdateCentroid();
+		EmitSelectionChange(OriginalSelection, CurrentSelection);
+		OnSelectionChanged.Broadcast();
+
+		// Rebuild after broadcast in case the outside world wants to adjust things like color...
+		RebuildDrawnElements(FTransform(GetCurrentSelectionCentroid()));
+	}
+}
+
+TSet<int32> UUVEditorMeshSelectionMechanic::RayCast(const FInputDeviceRay& Pos, EUVEditorMeshSelectionMode Mode)
+{
+	// TODO if the selection mode is Vertex, Edge or Triangle we'll return at most one index, if the selection mode is
+	// Component or Mesh we'll return a set of components. Perhaps it would be better to use TVariant here (and maybe
+	// also a TArray rather than a TSet)
+	TSet<int32> Result;
+	
+	int32 HitTid = IndexConstants::InvalidID;
+	for (int32 MeshIndex = 0; MeshIndex < MeshSpatials.Num(); ++MeshIndex)
+	{
+		FRay3d LocalRay(
+			(FVector3d)MeshTransforms[MeshIndex].InverseTransformPosition(Pos.WorldRay.Origin),
+			(FVector3d)MeshTransforms[MeshIndex].InverseTransformVector(Pos.WorldRay.Direction));
+
+		double RayT = 0;
+		if (MeshSpatials[MeshIndex]->FindNearestHitTriangle(LocalRay, RayT, HitTid))
+		{
+			CurrentSelection.Mesh = MeshSpatials[MeshIndex]->GetMesh();
+			CurrentSelection.TopologyTimestamp = CurrentSelection.Mesh->GetTopologyChangeStamp();
+			CurrentSelectionIndex = MeshIndex;
+			break;
+		}
+	}
+
+	if (HitTid != IndexConstants::InvalidID)
+	{
+		if (Mode == EUVEditorMeshSelectionMode::Component)
+		{
+			TRACE_CPUPROFILER_EVENT_SCOPE(Component);
+
+			FMeshConnectedComponents MeshSelectedComponent(CurrentSelection.Mesh);
+			TArray<int32> SeedTriangles;
+			SeedTriangles.Add(HitTid);
+			MeshSelectedComponent.FindTrianglesConnectedToSeeds(SeedTriangles);
+			ensure(MeshSelectedComponent.Components.Num() == 1); // Expect each triangle to only be in a single component
+			Result.Append(MoveTemp(MeshSelectedComponent.Components[0].Indices));
+		}
+		else if (Mode == EUVEditorMeshSelectionMode::Edge)
+		{
+			TRACE_CPUPROFILER_EVENT_SCOPE(Edge);
+			// TODO: We'll need the ability to hit occluded triangles to see if there is a better edge to snap to.
+
+			// Try to snap to one of the edges.
+			FIndex3i Eids = CurrentSelection.Mesh->GetTriEdges(HitTid);
+
+			FGeometrySet3 GeometrySet;
+			for (int i = 0; i < 3; ++i)
+			{
+				FIndex2i Vids = CurrentSelection.Mesh->GetEdgeV(Eids[i]);
+				FPolyline3d Polyline(CurrentSelection.Mesh->GetVertex(Vids.A), CurrentSelection.Mesh->GetVertex(Vids.B));
+				GeometrySet.AddCurve(Eids[i], Polyline);
+			}
+
+			FRay3d WorldRay((FVector3d)Pos.WorldRay.Origin, (FVector3d)Pos.WorldRay.Direction);
+			FGeometrySet3::FNearest Nearest;
+			if (GeometrySet.FindNearestCurveToRay(WorldRay, Nearest, 
+				[this](const FVector3d& Position1, const FVector3d& Position2) {
+					return ToolSceneQueriesUtil::PointSnapQuery(CameraState,
+						Position1, Position2,
+						ToolSceneQueriesUtil::GetDefaultVisualAngleSnapThreshD()); }))
+			{
+				Result = TSet<int32>{Nearest.ID};
+			}
+		}
+		else if (Mode == EUVEditorMeshSelectionMode::Vertex)
+		{
+			TRACE_CPUPROFILER_EVENT_SCOPE(Vertex);
+			// TODO: Improve this to handle super narrow, sliver triangles better, where testing near vertices can be difficult.
+
+			// Try to snap to one of the vertices
+			FIndex3i Vids = CurrentSelection.Mesh->GetTriangle(HitTid);
+
+			FGeometrySet3 GeometrySet;
+			for (int i = 0; i < 3; ++i)
+			{
+				GeometrySet.AddPoint(Vids[i], CurrentSelection.Mesh->GetTriVertex(HitTid, i));
+			}
+			
+			FRay3d WorldRay((FVector3d)Pos.WorldRay.Origin, (FVector3d)Pos.WorldRay.Direction);
+			FGeometrySet3::FNearest Nearest;
+			if (GeometrySet.FindNearestPointToRay(WorldRay, Nearest,
+				[this](const FVector3d& Position1, const FVector3d& Position2) {
+					return ToolSceneQueriesUtil::PointSnapQuery(CameraState,
+						Position1, Position2,
+						ToolSceneQueriesUtil::GetDefaultVisualAngleSnapThreshD()); }))
+			{
+				Result = TSet<int32>{Nearest.ID};
+			}
+		}
+		else if (Mode == EUVEditorMeshSelectionMode::Triangle)
+		{
+			TRACE_CPUPROFILER_EVENT_SCOPE(Triangle);
+
+			Result = TSet<int32>{HitTid};
+		}
+		else if (Mode == EUVEditorMeshSelectionMode::Mesh)
+		{
+			TRACE_CPUPROFILER_EVENT_SCOPE(Mesh);
+
+			for (int32 Tid : CurrentSelection.Mesh->TriangleIndicesItr())
+			{
+				Result.Add(Tid);
+			}
+		}
+		else
+		{
+			checkNoEntry();
+		}
+	}
+
+	return Result;
+}
+
+void UUVEditorMeshSelectionMechanic::OnDragRectangleStarted()
+{
+	TRACE_CPUPROFILER_EVENT_SCOPE(MeshSelectionMechanic_OnDragRectangleStarted); // Mark start of drag sequence
+	
+	if (CurrentSelection.Type != UVEditorMeshSelectionMechanicLocals::ToCompatibleDynamicMeshSelectionType(SelectionMode))
+	{
+		ClearCurrentSelection();	
+	}
+	
+	PreDragSelection = CurrentSelection;
+}
+
+void UUVEditorMeshSelectionMechanic::OnDragRectangleChanged(const FCameraRectangle& CurrentRectangle)
+{
+	TRACE_CPUPROFILER_EVENT_SCOPE(MeshSelectionMechanic_OnDragRectangleChanged);
+	
+	auto FindIDsInsideCurrentRectangle = [this](auto&& AddIDsToSelection)
+	{
+		TSet<int32> RectangleSelectedIDs;
+		
+		if (CurrentSelectionIndex == IndexConstants::InvalidID)
+		{
+			for (int32 MeshIndex = 0; MeshIndex < MeshSpatials.Num(); ++MeshIndex)
+			{
+				// TODO UUVEditorMeshSelectionMechanic is currently assuming that the provided mesh transform is identity, which is
+				//  the case in the UV editor, this restriction should be lifted. When we do this we should also apply the
+				//  (inverse) transform to the CurrentRectangle query rays. Search :ApplyTransformToQuery
+				ensure(MeshTransforms[MeshIndex].Identical(&FTransform::Identity, 0));
+				AddIDsToSelection(*MeshSpatials[MeshIndex], RectangleSelectedIDs);
+	
+				if (!RectangleSelectedIDs.IsEmpty())
+				{
+					// Pick the first selectable mesh for now, maybe we should try to be smarter
+					CurrentSelectionIndex = MeshIndex;
+					CurrentSelection.Mesh = MeshSpatials[CurrentSelectionIndex]->GetMesh();
+					break;
+				}
+			}
+		}
+		else
+		{
+			ensure(CurrentSelectionIndex != IndexConstants::InvalidID);
+			ensure(CurrentSelection.Mesh == MeshSpatials[CurrentSelectionIndex]->GetMesh());
+
+			// TODO See :ApplyTransformToQuery
+			ensure(MeshTransforms[CurrentSelectionIndex].Identical(&FTransform::Identity, 0));
+			AddIDsToSelection(*MeshSpatials[CurrentSelectionIndex], RectangleSelectedIDs);
+		}
+
+		return RectangleSelectedIDs;
+	};
+	using namespace UVEditorMeshSelectionMechanicLocals;
+	FAxisAlignedBox2d RectangleXY = GetRectangleXY(CurrentRectangle);
+	
+	TSet<int32> RectangleSelectedIDs;
+	if (SelectionMode == EUVEditorMeshSelectionMode::Vertex)
+	{
+		TRACE_CPUPROFILER_EVENT_SCOPE(MeshSelectionMechanic_OnDragRectangleChanged_Vertex);
+	
+		auto AddVertexIDsToSelection = [this, &RectangleXY](const FDynamicMeshAABBTree3& Tree, TSet<int32>& OutSelectedIDs)
+		{
+			TArray<int32> SelectedIDs =
+				FindAllIntersectionsAxisAlignedBox2(Tree, RectangleXY, AppendVertexIDs, AppendVertexIDsIfIntersected);
+			OutSelectedIDs.Append(MoveTemp(SelectedIDs));
+		};
+	
+		RectangleSelectedIDs = FindIDsInsideCurrentRectangle(AddVertexIDsToSelection);
+	}
+	else if (SelectionMode == EUVEditorMeshSelectionMode::Edge)
+	{
+		TRACE_CPUPROFILER_EVENT_SCOPE(MeshSelectionMechanic_OnDragRectangleChanged_Edge);
+		
+		auto AddEdgeIDsToSelection = [this, &RectangleXY](const FDynamicMeshAABBTree3& Tree, TSet<int32>& OutSelectedIDs)
+		{
+			TArray<int32> SelectedIDs =
+				FindAllIntersectionsAxisAlignedBox2(Tree, RectangleXY, AppendEdgeIDs, AppendEdgeIDsIfIntersected);
+			OutSelectedIDs.Append(MoveTemp(SelectedIDs));
+		};
+
+		RectangleSelectedIDs = FindIDsInsideCurrentRectangle(AddEdgeIDsToSelection);
+	}
+	else if (SelectionMode == EUVEditorMeshSelectionMode::Triangle)
+	{
+		TRACE_CPUPROFILER_EVENT_SCOPE(MeshSelectionMechanic_OnDragRectangleChanged_Triangle);
+		
+		auto AddTriangleIDsToSelection = [this, &RectangleXY] (const FDynamicMeshAABBTree3& Tree, TSet<int32>& OutSelectedIDs)
+		{
+			TArray<int32> SelectedIDs =
+				FindAllIntersectionsAxisAlignedBox2(Tree, RectangleXY, AppendTriangleID, AppendTriangleIDIfIntersected);
+			OutSelectedIDs.Append(MoveTemp(SelectedIDs));
+		};
+	
+		RectangleSelectedIDs = FindIDsInsideCurrentRectangle(AddTriangleIDsToSelection);
+	}
+	else if (SelectionMode == EUVEditorMeshSelectionMode::Component)
+	{
+		TRACE_CPUPROFILER_EVENT_SCOPE(MeshSelectionMechanic_OnDragRectangleChanged_Component);
+		
+		auto AddTriangleIDsToSelection = [this, &RectangleXY](const FDynamicMeshAABBTree3& Tree, TSet<int32>& OutSelectedIDs)
+		{
+			TArray<int32> SeedTriangles =
+				FindAllIntersectionsAxisAlignedBox2(Tree, RectangleXY, AppendTriangleID, AppendTriangleIDIfIntersected);
+
+			// TODO(Performance) For large meshes and selections following code is MUCH slower than AABB traversal,
+			//  consider precomputing the connected components an only updating them when the mesh topology changes
+			//  rather than every time the selection changes.
+			FMeshConnectedComponents MeshSelectedComponent(Tree.GetMesh());
+			MeshSelectedComponent.FindTrianglesConnectedToSeeds(SeedTriangles);
+			for (int ComponentIndex = 0; ComponentIndex < MeshSelectedComponent.Components.Num(); ComponentIndex++)
+			{
+				OutSelectedIDs.Append(MoveTemp(MeshSelectedComponent.Components[ComponentIndex].Indices));
+			}
+		};
+	
+		RectangleSelectedIDs = FindIDsInsideCurrentRectangle(AddTriangleIDsToSelection);
+	}
+	else if (SelectionMode == EUVEditorMeshSelectionMode::Mesh)
+	{
+		TRACE_CPUPROFILER_EVENT_SCOPE(MeshSelectionMechanic_OnDragRectangleChanged_Mesh);
+		
+		auto AddAllMeshTriangleIDsToSelection = [this, &RectangleXY] (const FDynamicMeshAABBTree3& Tree, TSet<int32>& OutSelectedIDs)
+		{
+			TArray<int32> SelectedIDs =
+				FindAllIntersectionsAxisAlignedBox2(Tree, RectangleXY, AppendTriangleID, AppendTriangleIDIfIntersected);
+			if (!SelectedIDs.IsEmpty())
+			{
+				for (int32 Tid : Tree.GetMesh()->TriangleIndicesItr())
+				{
+					OutSelectedIDs.Add(Tid);
+				}
+			}
+		};
+
+		RectangleSelectedIDs = FindIDsInsideCurrentRectangle(AddAllMeshTriangleIDsToSelection);
+	}
+	else
+	{
+		checkNoEntry();
+	}
+
+	UpdateCurrentSelection(RectangleSelectedIDs, true);
+	
+	// TODO(Performance) With large meshes and selections this call is much slower than AABB traversal (4x)
+	UpdateCentroid();
+	
+	// TODO(Performance) With large meshes and selections this call is MUCH slower than AABB traversal (60x)
+	RebuildDrawnElements(FTransform(GetCurrentSelectionCentroid()));
+}
+
+void UUVEditorMeshSelectionMechanic::OnDragRectangleFinished(const FCameraRectangle& CurrentRectangle, bool bCancelled)
+{
+	TRACE_CPUPROFILER_EVENT_SCOPE(MeshSelectionMechanic_OnDragRectangleFinished); // Mark end of drag sequence
+
+	// TODO(Performance) :DynamicMarqueeSelection Remove this call when marquee selection is fast enough to update
+	//  dynamically for large meshes
+	OnDragRectangleChanged(CurrentRectangle);
+
+	if (!bCancelled && (PreDragSelection != CurrentSelection))
+	{
+		UpdateCentroid();		
+		EmitSelectionChange(PreDragSelection, CurrentSelection);
+		OnSelectionChanged.Broadcast();
+
+		// Rebuild after broadcast in case the outside world wants to adjust things like color...
+		RebuildDrawnElements(FTransform(GetCurrentSelectionCentroid()));
+	}
+}
+
+void UUVEditorMeshSelectionMechanic::OnUpdateModifierState(int ModifierID, bool bIsOn)
+{
+	switch (ModifierID)
+	{
+	case ShiftModifierID:
+		bShiftToggle = bIsOn;
+		break;
+	case CtrlModifierID:
+		bCtrlToggle = bIsOn;
+		break;
+	default:
+		break;
+	}
+}
+
+#undef LOCTEXT_NAMESPACE
+
diff --git a/Engine/Plugins/Experimental/UVEditor/Source/UVEditorTools/Private/UVIslandConformalUnwrapAction.cpp b/Engine/Plugins/Experimental/UVEditor/Source/UVEditorTools/Private/UVIslandConformalUnwrapAction.cpp
--- a/Engine/Plugins/Experimental/UVEditor/Source/UVEditorTools/Private/UVIslandConformalUnwrapAction.cpp
+++ b/Engine/Plugins/Experimental/UVEditor/Source/UVEditorTools/Private/UVIslandConformalUnwrapAction.cpp
@@ -4,7 +4,7 @@
 
 #include "Algo/Unique.h"
 #include "ToolTargets/UVEditorToolMeshInput.h"
-#include "Selection/DynamicMeshSelection.h"
+#include "Selection/UVEditorDynamicMeshSelection.h"
 #include "MeshOpPreviewHelpers.h" // UMeshOpPreviewWithBackgroundCompute
 #include "PreviewMesh.h"
 #include "ToolSetupUtil.h"
@@ -19,7 +19,7 @@ using namespace UE::Geometry;
 
 UUVIslandConformalUnwrapAction::UUVIslandConformalUnwrapAction()
 {
-	CurrentSelection = MakeShared<UE::Geometry::FDynamicMeshSelection>();
+	CurrentSelection = MakeShared<UE::Geometry::FUVEditorDynamicMeshSelection>();
 }
 
 void UUVIslandConformalUnwrapAction::SetWorld(UWorld* WorldIn)
@@ -40,7 +40,7 @@ void UUVIslandConformalUnwrapAction::Shutdown()
 }
 
 
-void UUVIslandConformalUnwrapAction::SetSelection(int32 SelectionTargetIndexIn, const UE::Geometry::FDynamicMeshSelection* NewSelection)
+void UUVIslandConformalUnwrapAction::SetSelection(int32 SelectionTargetIndexIn, const UE::Geometry::FUVEditorDynamicMeshSelection* NewSelection)
 {
 	TRACE_CPUPROFILER_EVENT_SCOPE(UVIslandConformalUnwrapAction_SetSelection);
 	
@@ -51,7 +51,7 @@ void UUVIslandConformalUnwrapAction::SetSelection(int32 SelectionTargetIndexIn,
 	}
 	else
 	{
-		(*CurrentSelection) = UE::Geometry::FDynamicMeshSelection();
+		(*CurrentSelection) = UE::Geometry::FUVEditorDynamicMeshSelection();
 	}
 
 	UpdateVisualizations();
@@ -59,7 +59,7 @@ void UUVIslandConformalUnwrapAction::SetSelection(int32 SelectionTargetIndexIn,
 
 bool UUVIslandConformalUnwrapAction::PreCheckAction() 
 {
-	if (!CurrentSelection || CurrentSelection->Type != FDynamicMeshSelection::EType::Triangle
+	if (!CurrentSelection || CurrentSelection->Type != FUVEditorDynamicMeshSelection::EType::Triangle
 		|| CurrentSelection->Mesh == nullptr || SelectionTargetIndex == INDEX_NONE)
 	{
 		ParentTool->GetToolManager()->DisplayMessage(
diff --git a/Engine/Plugins/Experimental/UVEditor/Source/UVEditorTools/Private/UVSeamSewAction.cpp b/Engine/Plugins/Experimental/UVEditor/Source/UVEditorTools/Private/UVSeamSewAction.cpp
--- a/Engine/Plugins/Experimental/UVEditor/Source/UVEditorTools/Private/UVSeamSewAction.cpp
+++ b/Engine/Plugins/Experimental/UVEditor/Source/UVEditorTools/Private/UVSeamSewAction.cpp
@@ -6,7 +6,7 @@
 #include "Drawing/LineSetComponent.h"
 #include "Drawing/PreviewGeometryActor.h"
 #include "ToolTargets/UVEditorToolMeshInput.h"
-#include "Selection/DynamicMeshSelection.h"
+#include "Selection/UVEditorDynamicMeshSelection.h"
 #include "MeshOpPreviewHelpers.h" // UMeshOpPreviewWithBackgroundCompute
 #include "PreviewMesh.h"
 #include "ToolSetupUtil.h"
@@ -21,7 +21,7 @@ using namespace UE::Geometry;
 
 UUVSeamSewAction::UUVSeamSewAction()
 {
-	CurrentSelection = MakeShared<UE::Geometry::FDynamicMeshSelection>();
+	CurrentSelection = MakeShared<UE::Geometry::FUVEditorDynamicMeshSelection>();
 }
 
 void UUVSeamSewAction::SetWorld(UWorld* WorldIn)
@@ -75,7 +75,7 @@ void UUVSeamSewAction::Shutdown()
 }
 
 
-void UUVSeamSewAction::SetSelection(int32 SelectionTargetIndexIn, const UE::Geometry::FDynamicMeshSelection* NewSelection)
+void UUVSeamSewAction::SetSelection(int32 SelectionTargetIndexIn, const UE::Geometry::FUVEditorDynamicMeshSelection* NewSelection)
 {
 	TRACE_CPUPROFILER_EVENT_SCOPE(UVSeamSewAction_SetSelection);
 	
@@ -86,12 +86,12 @@ void UUVSeamSewAction::SetSelection(int32 SelectionTargetIndexIn, const UE::Geom
 	}
 	else
 	{
-		(*CurrentSelection) = UE::Geometry::FDynamicMeshSelection();
+		(*CurrentSelection) = UE::Geometry::FUVEditorDynamicMeshSelection();
 	}
 
 	EdgeSewCandidates.Reset();
 
-	if (SelectionTargetIndex != -1 && CurrentSelection && CurrentSelection->Type == FDynamicMeshSelection::EType::Edge)
+	if (SelectionTargetIndex != -1 && CurrentSelection && CurrentSelection->Type == FUVEditorDynamicMeshSelection::EType::Edge)
 	{
 		// TODO(Performance) This loop is very slow for large selections
 		TRACE_CPUPROFILER_EVENT_SCOPE(FindSewEdgeOppositePairing_Loop);
@@ -166,7 +166,7 @@ bool UUVSeamSewAction::PreCheckAction()
 		return false;
 	}
 
-	if (CurrentSelection->Type != FDynamicMeshSelection::EType::Edge)
+	if (CurrentSelection->Type != FUVEditorDynamicMeshSelection::EType::Edge)
 	{
 		ParentTool->GetToolManager()->DisplayMessage(
 			LOCTEXT("SewErrorSelectionNotEdge", "Cannot sew UVs. Selection was not an edge."),
diff --git a/Engine/Plugins/Experimental/UVEditor/Source/UVEditorTools/Private/UVSelectTool.cpp b/Engine/Plugins/Experimental/UVEditor/Source/UVEditorTools/Private/UVSelectTool.cpp
--- a/Engine/Plugins/Experimental/UVEditor/Source/UVEditorTools/Private/UVSelectTool.cpp
+++ b/Engine/Plugins/Experimental/UVEditor/Source/UVEditorTools/Private/UVSelectTool.cpp
@@ -17,8 +17,8 @@
 #include "MeshOpPreviewHelpers.h" // UMeshOpPreviewWithBackgroundCompute
 #include "Parameterization/DynamicMeshUVEditor.h"
 #include "PreviewMesh.h"
-#include "Selection/MeshSelectionMechanic.h"
-#include "Selection/DynamicMeshSelection.h"
+#include "Selection/UVEditorMeshSelectionMechanic.h"
+#include "Selection/UVEditorDynamicMeshSelection.h"
 #include "ToolSetupUtil.h"
 #include "UVEditorUXSettings.h"
 #include "EngineAnalytics.h"
@@ -42,9 +42,9 @@ namespace UVSelectToolLocals
 	// eids from our stored vid pairs.
 
 	/** If selection is a non-empty edge selection, update its eids using stored vid pairs. */
-	void UpdateSelectionEidsAfterMeshChange(FDynamicMeshSelection& SelectionInOut, TArray<FIndex2i>* VidPairsIn)
+	void UpdateSelectionEidsAfterMeshChange(FUVEditorDynamicMeshSelection& SelectionInOut, TArray<FIndex2i>* VidPairsIn)
 	{
-		if (!SelectionInOut.Mesh || SelectionInOut.Type != FDynamicMeshSelection::EType::Edge)
+		if (!SelectionInOut.Mesh || SelectionInOut.Type != FUVEditorDynamicMeshSelection::EType::Edge)
 		{
 			// No update necessary
 			return;
@@ -67,21 +67,21 @@ namespace UVSelectToolLocals
 	}
 
 	/** If selection mechanic holds a non-empty edge selection, update its eids using stored vid pairs. */
-	void UpdateSelectionEidsAfterMeshChange(UMeshSelectionMechanic& SelectionMechanic, TArray<FIndex2i>* VidPairsIn)
+	void UpdateSelectionEidsAfterMeshChange(UUVEditorMeshSelectionMechanic& SelectionMechanic, TArray<FIndex2i>* VidPairsIn)
 	{
-		const FDynamicMeshSelection& CurrentSelection = SelectionMechanic.GetCurrentSelection();
-		if (CurrentSelection.Mesh && CurrentSelection.Type == FDynamicMeshSelection::EType::Edge)
+		const FUVEditorDynamicMeshSelection& CurrentSelection = SelectionMechanic.GetCurrentSelection();
+		if (CurrentSelection.Mesh && CurrentSelection.Type == FUVEditorDynamicMeshSelection::EType::Edge)
 		{
-			FDynamicMeshSelection UpdatedSelection = CurrentSelection;
+			FUVEditorDynamicMeshSelection UpdatedSelection = CurrentSelection;
 			UpdateSelectionEidsAfterMeshChange(UpdatedSelection, VidPairsIn);
 			SelectionMechanic.SetSelection(UpdatedSelection, false, false);
 		}
 	}
 
-	void GetVidPairsFromSelection(const FDynamicMeshSelection& SelectionIn, TArray<FIndex2i>& VidPairsOut)
+	void GetVidPairsFromSelection(const FUVEditorDynamicMeshSelection& SelectionIn, TArray<FIndex2i>& VidPairsOut)
 	{
 		VidPairsOut.Reset();
-		if (!SelectionIn.Mesh || SelectionIn.Type != FDynamicMeshSelection::EType::Edge)
+		if (!SelectionIn.Mesh || SelectionIn.Type != FUVEditorDynamicMeshSelection::EType::Edge)
 		{
 			// No vid pairs to add
 			return;
@@ -114,8 +114,8 @@ namespace UVSelectToolLocals
 		 * @param EdgeVidPairsBeforeIn
 		 * @param EdgeVidPairsAfterIn
 		 */
-		FSelectionChange(const FDynamicMeshSelection& SelectionBeforeIn,
-			const FDynamicMeshSelection& SelectionAfterIn,
+		FSelectionChange(const FUVEditorDynamicMeshSelection& SelectionBeforeIn,
+			const FUVEditorDynamicMeshSelection& SelectionAfterIn,
 			const FTransform& GizmoBeforeIn,
 			TUniquePtr<TArray<FIndex2i>> EdgeVidPairsBeforeIn,
 			TUniquePtr<TArray<FIndex2i>> EdgeVidPairsAfterIn
@@ -128,8 +128,8 @@ namespace UVSelectToolLocals
 		{
 			// Make sure that for both selections, if we have a non-empty edge selection, we have vid pairs.
 			ensure(!(
-				(SelectionBefore.Mesh && SelectionBefore.Type == FDynamicMeshSelection::EType::Edge && !EdgeVidPairsBefore)
-				|| (SelectionAfter.Mesh && SelectionAfter.Type == FDynamicMeshSelection::EType::Edge && !EdgeVidPairsAfter)));
+				(SelectionBefore.Mesh && SelectionBefore.Type == FUVEditorDynamicMeshSelection::EType::Edge && !EdgeVidPairsBefore)
+				|| (SelectionAfter.Mesh && SelectionAfter.Type == FUVEditorDynamicMeshSelection::EType::Edge && !EdgeVidPairsAfter)));
 		}
 
 		virtual void Apply(UObject* Object) override
@@ -165,8 +165,8 @@ namespace UVSelectToolLocals
 		}
 
 	protected:
-		FDynamicMeshSelection SelectionBefore;
-		FDynamicMeshSelection SelectionAfter;
+		FUVEditorDynamicMeshSelection SelectionBefore;
+		FUVEditorDynamicMeshSelection SelectionAfter;
 		FTransform GizmoBefore;
 
 		TUniquePtr<TArray<FIndex2i>> EdgeVidPairsBefore;
@@ -314,31 +314,10 @@ void UUVSelectTool::Setup()
 	ToolActions->Initialize(this);
 	AddToolPropertySource(ToolActions);
 
-	SelectionMechanic = NewObject<UMeshSelectionMechanic>();
+	SelectionMechanic = NewObject<UUVEditorMeshSelectionMechanic>();
 	SelectionMechanic->Setup(this);
 	SelectionMechanic->SetWorld(Targets[0]->UnwrapPreview->GetWorld());
 	SelectionMechanic->OnSelectionChanged.AddUObject(this, &UUVSelectTool::OnSelectionChanged);
-	
-	FMeshSelectionMechanicStyle SelectionStyle;
-	
-	SelectionStyle.TriangleColor = FUVEditorUXSettings::SelectionTriangleFillColor;
-	SelectionStyle.LineColor = FUVEditorUXSettings::SelectionTriangleWireframeColor;
-	SelectionStyle.PointColor = FUVEditorUXSettings::SelectionTriangleWireframeColor;
-	SelectionStyle.TriangleOpacity = FUVEditorUXSettings::SelectionTriangleOpacity;
-	SelectionStyle.LineThickness = FUVEditorUXSettings::SelectionLineThickness;
-	SelectionStyle.PointThickness = FUVEditorUXSettings::SelectionPointThickness;
-	SelectionStyle.LineAndPointDepthBias = FUVEditorUXSettings::SelectionWireframeDepthBias;
-	SelectionStyle.TriangleDepthBias = FUVEditorUXSettings::SelectionTriangleDepthBias;
-	
-	SelectionStyle.HoverLineAndPointDepthBias = FUVEditorUXSettings::SelectionHoverWireframeDepthBias;
-	SelectionStyle.HoverTriangleDepthBias = FUVEditorUXSettings::SelectionHoverTriangleDepthBias;
-	SelectionStyle.HoverPointColor = FUVEditorUXSettings::SelectionHoverTriangleWireframeColor;
-	SelectionStyle.HoverEdgeColor = FUVEditorUXSettings::SelectionHoverTriangleWireframeColor;
-	SelectionStyle.HoverTriangleEdgeColor = FUVEditorUXSettings::SelectionHoverTriangleWireframeColor;
-	SelectionStyle.HoverTriangleFillColor = FUVEditorUXSettings::SelectionHoverTriangleFillColor;
-	SelectionStyle.HoverTriangleOpacity = FUVEditorUXSettings::SelectionHoverTriangleOpacity;
-	
-	SelectionMechanic->SetVisualizationStyle(SelectionStyle);
 		
 	// Make it so that our selection mechanic creates undo/redo transactions that go to a selection
 	// change router, which we use to route to the current selection mechanic on each tool invocation.
@@ -350,17 +329,17 @@ void UUVSelectTool::Setup()
 	}
 	ChangeRouter->CurrentSelectTool = this;
 
-	SelectionMechanic->EmitSelectionChange = [this](const FDynamicMeshSelection& OldSelection,
-		const FDynamicMeshSelection& NewSelection)
+	SelectionMechanic->EmitSelectionChange = [this](const FUVEditorDynamicMeshSelection& OldSelection,
+		const FUVEditorDynamicMeshSelection& NewSelection)
 	{
 		TUniquePtr<TArray<FIndex2i>> VidPairsBefore;
 		TUniquePtr<TArray<FIndex2i>> VidPairsAfter;
-		if (OldSelection.Type == FDynamicMeshSelection::EType::Edge)
+		if (OldSelection.Type == FUVEditorDynamicMeshSelection::EType::Edge)
 		{
 			VidPairsBefore = MakeUnique<TArray<FIndex2i>>();
 			GetVidPairsFromSelection(OldSelection, *VidPairsBefore);
 		}
-		if (NewSelection.Type == FDynamicMeshSelection::EType::Edge)
+		if (NewSelection.Type == FUVEditorDynamicMeshSelection::EType::Edge)
 		{
 			VidPairsAfter = MakeUnique<TArray<FIndex2i>>();
 			GetVidPairsFromSelection(NewSelection, *VidPairsAfter);
@@ -500,7 +479,7 @@ void UUVSelectTool::Shutdown(EToolShutdownType ShutdownType)
 	if (!SelectionMechanic->GetCurrentSelection().IsEmpty())
 	{
 		// (The broadcast here is so that we still broadcast on undo)
-		SelectionMechanic->SetSelection(UMeshSelectionMechanic::FDynamicMeshSelection(), true, true);
+		SelectionMechanic->SetSelection(UUVEditorMeshSelectionMechanic::FUVEditorDynamicMeshSelection(), true, true);
 	}
 
 	ChangeRouter->CurrentSelectTool = nullptr;
@@ -546,7 +525,7 @@ void UUVSelectTool::Shutdown(EToolShutdownType ShutdownType)
 	RecordAnalytics();
 }
 
-void UUVSelectTool::SetSelection(const FDynamicMeshSelection& NewSelection)
+void UUVSelectTool::SetSelection(const FUVEditorDynamicMeshSelection& NewSelection)
 {
 	SelectionMechanic->SetSelection(NewSelection, true, 
 		false); // Don't emit undo because this function is called from undo
@@ -563,27 +542,27 @@ void UUVSelectTool::SetSelection(const FDynamicMeshSelection& NewSelection)
 	UUVToolViewportButtonsAPI::ESelectionMode CurrentMode = ViewportButtonsAPI->GetSelectionMode();
 	switch (NewSelection.Type)
 	{
-	case FDynamicMeshSelection::EType::Vertex:
+	case FUVEditorDynamicMeshSelection::EType::Vertex:
 		if (CurrentMode != UUVToolViewportButtonsAPI::ESelectionMode::Vertex)
 		{
 			ViewportButtonsAPI->SetSelectionMode(UUVToolViewportButtonsAPI::ESelectionMode::Vertex, false);
-			SelectionMechanic->SelectionMode = EMeshSelectionMechanicMode::Vertex;
+			SelectionMechanic->SelectionMode = EUVEditorMeshSelectionMode::Vertex;
 		}
 		break;
-	case FDynamicMeshSelection::EType::Edge:
+	case FUVEditorDynamicMeshSelection::EType::Edge:
 		if (CurrentMode != UUVToolViewportButtonsAPI::ESelectionMode::Edge)
 		{
 			ViewportButtonsAPI->SetSelectionMode(UUVToolViewportButtonsAPI::ESelectionMode::Edge, false);
-			SelectionMechanic->SelectionMode = EMeshSelectionMechanicMode::Edge;
+			SelectionMechanic->SelectionMode = EUVEditorMeshSelectionMode::Edge;
 		}
 		break;
-	case FDynamicMeshSelection::EType::Triangle:
+	case FUVEditorDynamicMeshSelection::EType::Triangle:
 		if (CurrentMode != UUVToolViewportButtonsAPI::ESelectionMode::Triangle
 			&& CurrentMode != UUVToolViewportButtonsAPI::ESelectionMode::Island
 			&& CurrentMode != UUVToolViewportButtonsAPI::ESelectionMode::Mesh)
 		{
 			ViewportButtonsAPI->SetSelectionMode(UUVToolViewportButtonsAPI::ESelectionMode::Triangle, false);
-			SelectionMechanic->SelectionMode = EMeshSelectionMechanicMode::Triangle;
+			SelectionMechanic->SelectionMode = EUVEditorMeshSelectionMode::Triangle;
 		}
 		break;
 	}
@@ -602,7 +581,7 @@ void UUVSelectTool::OnPropertyModified(UObject* PropertySet, FProperty* Property
 
 void UUVSelectTool::UpdateGizmo()
 {
-	const FDynamicMeshSelection& Selection = SelectionMechanic->GetCurrentSelection();
+	const FUVEditorDynamicMeshSelection& Selection = SelectionMechanic->GetCurrentSelection();
 
 	if (!Selection.IsEmpty())
 	{
@@ -618,28 +597,28 @@ void UUVSelectTool::UpdateGizmo()
 
 void UUVSelectTool::UpdateSelectionMode()
 {
-	EMeshSelectionMechanicMode TargetMode;
+	EUVEditorMeshSelectionMode TargetMode;
 	switch (ViewportButtonsAPI->GetSelectionMode())
 	{
 	case UUVToolViewportButtonsAPI::ESelectionMode::Vertex:
-		TargetMode = EMeshSelectionMechanicMode::Vertex;
+		TargetMode = EUVEditorMeshSelectionMode::Vertex;
 		break;	
 	case UUVToolViewportButtonsAPI::ESelectionMode::Edge:
-		TargetMode = EMeshSelectionMechanicMode::Edge;
+		TargetMode = EUVEditorMeshSelectionMode::Edge;
 		break;
 	case UUVToolViewportButtonsAPI::ESelectionMode::Triangle:
-		TargetMode = EMeshSelectionMechanicMode::Triangle;
+		TargetMode = EUVEditorMeshSelectionMode::Triangle;
 		break;
 	case UUVToolViewportButtonsAPI::ESelectionMode::Island:
-		TargetMode = EMeshSelectionMechanicMode::Component;
+		TargetMode = EUVEditorMeshSelectionMode::Component;
 		break;
 	case UUVToolViewportButtonsAPI::ESelectionMode::Mesh:
-		TargetMode = EMeshSelectionMechanicMode::Mesh;
+		TargetMode = EUVEditorMeshSelectionMode::Mesh;
 		break;
 	default:
 		// We shouldn't ever get "none" as the selection mode...
 		ensure(false);
-		TargetMode = EMeshSelectionMechanicMode::Vertex;
+		TargetMode = EUVEditorMeshSelectionMode::Vertex;
 		break;
 	}
 	SelectionMechanic->ChangeSelectionMode(TargetMode); // broadcast and emit undo if needed
@@ -653,7 +632,7 @@ void UUVSelectTool::OnSelectionChanged()
 
 	ClearWarning();
 
-	const FDynamicMeshSelection& Selection = SelectionMechanic->GetCurrentSelection();
+	const FUVEditorDynamicMeshSelection& Selection = SelectionMechanic->GetCurrentSelection();
 
 	GetVidPairsFromSelection(Selection, CurrentSelectionVidPairs);
 
@@ -682,7 +661,7 @@ void UUVSelectTool::OnSelectionChanged()
 		// Note the selected vids
 		TSet<int32> VidSet;
 		TSet<int32> TidSet;
-		if (Selection.Type == FDynamicMeshSelection::EType::Triangle)
+		if (Selection.Type == FUVEditorDynamicMeshSelection::EType::Triangle)
 		{
 			TRACE_CPUPROFILER_EVENT_SCOPE(Triangle);
 
@@ -719,7 +698,7 @@ void UUVSelectTool::OnSelectionChanged()
 				}
 			}
 		}
-		else if (Selection.Type == FDynamicMeshSelection::EType::Edge)
+		else if (Selection.Type == FUVEditorDynamicMeshSelection::EType::Edge)
 		{
 			TRACE_CPUPROFILER_EVENT_SCOPE(Edge);
 
@@ -752,7 +731,7 @@ void UUVSelectTool::OnSelectionChanged()
 					Target->UnwrapVidToAppliedVid(EdgeVids.B)));
 			}
 		}
-		else if (Selection.Type == FDynamicMeshSelection::EType::Vertex)
+		else if (Selection.Type == FUVEditorDynamicMeshSelection::EType::Vertex)
 		{
 			TRACE_CPUPROFILER_EVENT_SCOPE(Vertex);
 
@@ -803,7 +782,7 @@ void UUVSelectTool::UpdateLivePreviewLines()
 	LivePreviewLineSet->Clear();
 	LivePreviewPointSet->Clear();
 
-	const FDynamicMeshSelection& Selection = SelectionMechanic->GetCurrentSelection();
+	const FUVEditorDynamicMeshSelection& Selection = SelectionMechanic->GetCurrentSelection();
 	if (!Selection.IsEmpty())
 	{
 		FTransform MeshTransform = Targets[SelectionTargetIndex]->AppliedPreview->PreviewMesh->GetTransform();
@@ -973,7 +952,7 @@ void UUVSelectTool::RequestAction(ESelectToolAction ActionType)
 
 void UUVSelectTool::ApplyAction(ESelectToolAction ActionType)
 {
-	auto MaybeAddAnalyticsActionHistoryItem = [this, &ActionType](FDynamicMeshSelection::EType ExpectedType)
+	auto MaybeAddAnalyticsActionHistoryItem = [this, &ActionType](FUVEditorDynamicMeshSelection::EType ExpectedType)
 	{
 		if (!SelectionMechanic->GetCurrentSelection().IsEmpty() && (SelectionMechanic->GetCurrentSelection().Type == ExpectedType))
 		{
@@ -992,12 +971,12 @@ void UUVSelectTool::ApplyAction(ESelectToolAction ActionType)
 		{
 			TRACE_CPUPROFILER_EVENT_SCOPE(ApplyAction_Sew);
 	
-			MaybeAddAnalyticsActionHistoryItem(FDynamicMeshSelection::EType::Edge);
+			MaybeAddAnalyticsActionHistoryItem(FUVEditorDynamicMeshSelection::EType::Edge);
 
 			const FText TransactionName(LOCTEXT("SewCompleteTransactionName", "Sew Edges"));
 			EmitChangeAPI->BeginUndoTransaction(TransactionName);
 
-			SelectionMechanic->SetSelection(FDynamicMeshSelection(), false, true);
+			SelectionMechanic->SetSelection(FUVEditorDynamicMeshSelection(), false, true);
 			SewAction->ExecuteAction(*EmitChangeAPI);
 
 			EmitChangeAPI->EndUndoTransaction();
@@ -1008,12 +987,12 @@ void UUVSelectTool::ApplyAction(ESelectToolAction ActionType)
 		{
 			TRACE_CPUPROFILER_EVENT_SCOPE(ApplyAction_IslandConformalUnwrap);
 			
-			MaybeAddAnalyticsActionHistoryItem(FDynamicMeshSelection::EType::Triangle);
+			MaybeAddAnalyticsActionHistoryItem(FUVEditorDynamicMeshSelection::EType::Triangle);
 
 			const FText TransactionName(LOCTEXT("ConformalUnwrapCompleteTransactionName", "Conformal Unwrap Islands"));
 			EmitChangeAPI->BeginUndoTransaction(TransactionName);
 
-			SelectionMechanic->SetSelection(FDynamicMeshSelection(), false, true);
+			SelectionMechanic->SetSelection(FUVEditorDynamicMeshSelection(), false, true);
 			IslandConformalUnwrapAction->ExecuteAction(*EmitChangeAPI);
 
 			EmitChangeAPI->EndUndoTransaction();
@@ -1021,7 +1000,7 @@ void UUVSelectTool::ApplyAction(ESelectToolAction ActionType)
 		break;
 	case ESelectToolAction::Split:
 		
-		MaybeAddAnalyticsActionHistoryItem(FDynamicMeshSelection::EType::Edge);
+		MaybeAddAnalyticsActionHistoryItem(FUVEditorDynamicMeshSelection::EType::Edge);
 		
 		ApplySplit();
 	default:
@@ -1033,7 +1012,7 @@ void UUVSelectTool::ApplySplit()
 {
 	using namespace UVSelectToolLocals;
 
-	const FDynamicMeshSelection& Selection = SelectionMechanic->GetCurrentSelection();
+	const FUVEditorDynamicMeshSelection& Selection = SelectionMechanic->GetCurrentSelection();
 
 	if (Selection.IsEmpty())
 	{
@@ -1041,11 +1020,11 @@ void UUVSelectTool::ApplySplit()
 			LOCTEXT("SplitErrorSelectionEmpty", "Cannot split UVs. Selection was empty."),
 			EToolMessageLevel::UserWarning);
 	}
-	else if (Selection.Type == FDynamicMeshSelection::EType::Edge)
+	else if (Selection.Type == FUVEditorDynamicMeshSelection::EType::Edge)
 	{
 		ApplySplitEdges();
 	}
-	else if (Selection.Type == FDynamicMeshSelection::EType::Vertex)
+	else if (Selection.Type == FUVEditorDynamicMeshSelection::EType::Vertex)
 	{
 		ApplySplitBowtieVertices();
 	}
@@ -1059,8 +1038,8 @@ void UUVSelectTool::ApplySplit()
 
 void UUVSelectTool::ApplySplitEdges()
 {
-	const FDynamicMeshSelection& Selection = SelectionMechanic->GetCurrentSelection();
-	if (!ensure(SelectionTargetIndex >= 0 && !Selection.IsEmpty() && Selection.Type == FDynamicMeshSelection::EType::Edge))
+	const FUVEditorDynamicMeshSelection& Selection = SelectionMechanic->GetCurrentSelection();
+	if (!ensure(SelectionTargetIndex >= 0 && !Selection.IsEmpty() && Selection.Type == FUVEditorDynamicMeshSelection::EType::Edge))
 	{
 		return;
 	}
@@ -1110,7 +1089,7 @@ void UUVSelectTool::ApplySplitEdges()
 	// Clear selection here so that it is restored on undo
 	// TODO: We could apply a new selection after the change, but this should only happen
 	// once we apply a new selection on split, which we currently don't do.
-	SelectionMechanic->SetSelection(FDynamicMeshSelection(), true, true);
+	SelectionMechanic->SetSelection(FUVEditorDynamicMeshSelection(), true, true);
 
 	// Perform the update
 	TArray<int32> AppliedTids = TidSet.Array();
@@ -1128,8 +1107,8 @@ void UUVSelectTool::ApplySplitEdges()
 
 void UUVSelectTool::ApplySplitBowtieVertices()
 {
-	const FDynamicMeshSelection& Selection = SelectionMechanic->GetCurrentSelection();
-	if (!ensure(SelectionTargetIndex >= 0 && !Selection.IsEmpty() && Selection.Type == FDynamicMeshSelection::EType::Vertex))
+	const FUVEditorDynamicMeshSelection& Selection = SelectionMechanic->GetCurrentSelection();
+	if (!ensure(SelectionTargetIndex >= 0 && !Selection.IsEmpty() && Selection.Type == FUVEditorDynamicMeshSelection::EType::Vertex))
 	{
 		return;
 	}
@@ -1168,10 +1147,10 @@ void UUVSelectTool::ApplySplitBowtieVertices()
 	EmitChangeAPI->BeginUndoTransaction(TransactionName);
 
 	// Emit selection clear first so that we restore it on undo
-	FDynamicMeshSelection NewSelection = SelectionMechanic->GetCurrentSelection(); // save type, etc
+	FUVEditorDynamicMeshSelection NewSelection = SelectionMechanic->GetCurrentSelection(); // save type, etc
 	// TODO: This emitted transaction doesn't actually need to broadcast on redo, but we don't yet
 	// have support for that.
-	SelectionMechanic->SetSelection(FDynamicMeshSelection(), false, true); // don't broadcast, do emit undo
+	SelectionMechanic->SetSelection(FUVEditorDynamicMeshSelection(), false, true); // don't broadcast, do emit undo
 
 	// Perform the update
 	TArray<int32> AppliedTids = TidSet.Array();
diff --git a/Engine/Plugins/Experimental/UVEditor/Source/UVEditorTools/Public/Selection/UVEditorDynamicMeshSelection.h b/Engine/Plugins/Experimental/UVEditor/Source/UVEditorTools/Public/Selection/UVEditorDynamicMeshSelection.h
new file mode 100644
--- /dev/null
+++ b/Engine/Plugins/Experimental/UVEditor/Source/UVEditorTools/Public/Selection/UVEditorDynamicMeshSelection.h
@@ -0,0 +1,65 @@
+// Copyright Epic Games, Inc. All Rights Reserved.
+
+#pragma once
+
+#include "CoreMinimal.h"
+
+namespace UE {
+namespace Geometry {
+
+class FDynamicMesh3;
+
+/**
+ * Class that represents a selection in a dynamic mesh, for use in the UV editor.
+ */
+class UVEDITORTOOLS_API FUVEditorDynamicMeshSelection
+{
+public:
+	enum class EType
+	{
+		Vertex,
+		Edge,
+		Triangle,
+		Group
+	};
+
+	const FDynamicMesh3* Mesh = nullptr;
+	TSet<int32> SelectedIDs;
+	EType Type = EType::Vertex;
+
+	// Not relevant if the selection type is not Group
+	int32 GroupLayer = 0;
+
+	// Can be used to discard the selection if topology of the mesh
+	// has changed (to avoid the risk of referencing elements that may
+	// have been deleted)
+	int32 TopologyTimestamp = -1;
+
+	/** Checks whether the selection's timestamp still matches the meshes topology timestamp. */
+	bool MatchesTimestamp() const
+	{
+		return Mesh && Mesh->GetTopologyChangeStamp() == TopologyTimestamp;
+	}
+
+	bool IsEmpty() const
+	{
+		return SelectedIDs.IsEmpty();
+	}
+
+	bool operator==(const FUVEditorDynamicMeshSelection& Other)
+	{
+		return Mesh == Other.Mesh
+			&& Type == Other.Type
+			&& TopologyTimestamp == Other.TopologyTimestamp
+			&& (Type != EType::Group || GroupLayer == Other.GroupLayer)
+			&& SelectedIDs.Num() == Other.SelectedIDs.Num()
+			&& SelectedIDs.Includes(Other.SelectedIDs);
+	}
+
+	bool operator!=(const FUVEditorDynamicMeshSelection& Other)
+	{
+		return !(*this == Other);
+	}
+};
+
+}}// end namespace UE::Geometry
\ No newline at end of file
diff --git a/Engine/Plugins/Experimental/UVEditor/Source/UVEditorTools/Public/Selection/UVEditorMeshSelectionMechanic.h b/Engine/Plugins/Experimental/UVEditor/Source/UVEditorTools/Public/Selection/UVEditorMeshSelectionMechanic.h
new file mode 100644
--- /dev/null
+++ b/Engine/Plugins/Experimental/UVEditor/Source/UVEditorTools/Public/Selection/UVEditorMeshSelectionMechanic.h
@@ -0,0 +1,177 @@
+// Copyright Epic Games, Inc. All Rights Reserved.
+
+#pragma once
+
+#include "CoreMinimal.h"
+#include "GeometryBase.h" // Predeclare macros
+
+#include "BaseBehaviors/BehaviorTargetInterfaces.h"
+#include "DynamicMesh/DynamicMeshAABBTree3.h"
+#include "InteractionMechanic.h"
+#include "InteractiveTool.h"
+#include "Selection/UVEditorDynamicMeshSelection.h"
+#include "Mechanics/RectangleMarqueeMechanic.h"
+#include "ToolContextInterfaces.h" //FViewCameraState
+
+#include "UVEditorMeshSelectionMechanic.generated.h"
+
+class UTriangleSetComponent;
+class ULineSetComponent;
+class UPointSetComponent;
+class APreviewGeometryActor;
+class UMaterialInstanceDynamic;
+struct FCameraRectangle;
+
+enum class EUVEditorMeshSelectionMode
+{
+	Component,
+	
+	// Not yet fully implemented for UV mesh purposes, since 
+	// we need to be able to select occluded edges
+	Edge,
+	Vertex,
+	Triangle,
+
+	// TODO: This might be good to rename later. And determine how it might interact with multi-mesh selection?
+	Mesh
+};
+
+UCLASS()
+class UVEDITORTOOLS_API UUVEditorMeshSelectionMechanicProperties : public UInteractiveToolPropertySet
+{
+	GENERATED_BODY()
+
+public:
+	/** If true, show the currently hovered selectable primitive */
+	UPROPERTY(EditAnywhere, Category = Options)
+	bool bShowHoveredElements = true;
+};
+
+/**
+ * Mechanic for selecting elements of a dynamic mesh in the UV editor.
+ * 
+ * TODO: Currently only able to select unoccluded elements.
+ */
+UCLASS()
+class UVEDITORTOOLS_API UUVEditorMeshSelectionMechanic : public UInteractionMechanic, public IClickBehaviorTarget
+{
+	GENERATED_BODY()
+
+public:
+	using FDynamicMeshAABBTree3 = UE::Geometry::FDynamicMeshAABBTree3;
+	using FUVEditorDynamicMeshSelection = UE::Geometry::FUVEditorDynamicMeshSelection;
+
+	virtual ~UUVEditorMeshSelectionMechanic() {}
+
+	virtual void Setup(UInteractiveTool* ParentTool) override;
+	virtual void Shutdown() override;
+
+	void SetWorld(UWorld* World);
+
+	// Use this to initialize the meshes we want to hit test.
+	virtual void AddSpatial(TSharedPtr<FDynamicMeshAABBTree3> SpatialIn, const FTransform& TransformIn);
+
+	FVector3d GetCurrentSelectionCentroid();
+
+	// Rebuilds the drawn selection highlights, and intializes them in such a way that their transform
+	// is equal to StartTransform (useful so that their transform can later be changed)
+	void RebuildDrawnElements(const FTransform& StartTransform);
+
+	// Changes the transform of the selection highlights. Useful for quickly updating the hightlight
+	// without rebuilding it, when the change is a transformation.
+	void SetDrawnElementsTransform(const FTransform& Transform);
+
+	virtual const FUVEditorDynamicMeshSelection& GetCurrentSelection() const;
+	void ChangeSelectionMode(const EUVEditorMeshSelectionMode& TargetMode);
+	virtual void SetSelection(const FUVEditorDynamicMeshSelection& Selection, bool bBroadcast = false, bool bEmitChange = false);
+
+	virtual void Render(IToolsContextRenderAPI* RenderAPI) override;
+	virtual void DrawHUD(FCanvas* Canvas, IToolsContextRenderAPI* RenderAPI);
+
+	void OnDragRectangleStarted();
+	void OnDragRectangleChanged(const FCameraRectangle& CurrentRectangle);
+	void OnDragRectangleFinished(const FCameraRectangle& Rectangle, bool bCancelled);
+
+	// IClickBehaviorTarget implementation
+	virtual FInputRayHit IsHitByClick(const FInputDeviceRay& ClickPos) override;
+	virtual void OnClicked(const FInputDeviceRay& ClickPos) override;
+
+	// IModifierToggleBehaviorTarget implementation
+	virtual void OnUpdateModifierState(int ModifierID, bool bIsOn) override;
+	
+	TSet<int32> RayCast(const FInputDeviceRay& ClickPos, EUVEditorMeshSelectionMode Mode);
+
+	FSimpleMulticastDelegate OnSelectionChanged;
+
+	EUVEditorMeshSelectionMode SelectionMode;
+
+	/**
+	 * Function to use for emitting selection change events. If not set, Setup() will attach a version that
+	 * uses EmitObjectChange() on the tool manager to emit a change that operates on this mechanic.
+	 * TODO: User should probably be able to specify whether the emitted transaction should broadcast
+	 * OnSelectionChanged on redo, undo, or both, to allow selection change events to be used as bookends
+	 * around topology changes.
+	 */ 
+	TUniqueFunction<void(const FUVEditorDynamicMeshSelection& OldSelection, const FUVEditorDynamicMeshSelection& NewSelection)> EmitSelectionChange;
+
+protected:
+
+	// All four combinations of shift/ctrl down are assigned a behaviour
+	bool ShouldAddToSelection() const { return !bCtrlToggle && bShiftToggle; }
+	bool ShouldRemoveFromSelection() const { return bCtrlToggle && !bShiftToggle; }
+	bool ShouldToggleFromSelection() const { return bCtrlToggle && bShiftToggle; }
+	bool ShouldRestartSelection() const { return !bCtrlToggle && !bShiftToggle; }
+
+	UPROPERTY()
+	TObjectPtr<UUVEditorMeshSelectionMechanicProperties> Settings = nullptr;
+
+	UPROPERTY()
+	TObjectPtr<URectangleMarqueeMechanic> MarqueeMechanic;
+
+	UPROPERTY()
+	TObjectPtr<APreviewGeometryActor> PreviewGeometryActor = nullptr;
+
+	/** The material being displayed for selected triangles */
+	UPROPERTY()
+	TObjectPtr<UMaterialInstanceDynamic> TriangleSetMaterial = nullptr;
+	UPROPERTY()
+	TObjectPtr<UTriangleSetComponent> TriangleSet = nullptr;
+	UPROPERTY()
+	TObjectPtr<ULineSetComponent> LineSet = nullptr;
+	UPROPERTY()
+	TObjectPtr<UPointSetComponent> PointSet = nullptr;
+	
+	UPROPERTY()
+	TObjectPtr<APreviewGeometryActor> HoverGeometryActor = nullptr;
+	
+	UPROPERTY()
+	TObjectPtr<UMaterialInstanceDynamic> HoverTriangleSetMaterial = nullptr;
+	UPROPERTY()
+	TObjectPtr<UTriangleSetComponent> HoverTriangleSet = nullptr;
+	UPROPERTY()
+	TObjectPtr<ULineSetComponent> HoverLineSet = nullptr;
+	UPROPERTY()
+	TObjectPtr<UPointSetComponent> HoverPointSet = nullptr;
+	
+
+	TArray<TSharedPtr<FDynamicMeshAABBTree3>> MeshSpatials;
+	TArray<FTransform> MeshTransforms;
+	FUVEditorDynamicMeshSelection CurrentSelection;
+	FUVEditorDynamicMeshSelection PreDragSelection;
+	int32 CurrentSelectionIndex = IndexConstants::InvalidID;
+	FViewCameraState CameraState;
+
+	bool bShiftToggle = false;
+	bool bCtrlToggle = false;
+	static const int32 ShiftModifierID = 1;
+	static const int32 CtrlModifierID = 2;
+
+	FVector3d CurrentSelectionCentroid;
+	int32 CentroidTimestamp = -1;
+	void UpdateCentroid();
+
+private:
+	void ClearCurrentSelection();
+	void UpdateCurrentSelection(const TSet<int32>& NewSelection, bool CalledFromOnRectangleChanged = false);
+};
+
diff --git a/Engine/Plugins/Experimental/UVEditor/Source/UVEditorTools/Public/UVIslandConformalUnwrapAction.h b/Engine/Plugins/Experimental/UVEditor/Source/UVEditorTools/Public/UVIslandConformalUnwrapAction.h
--- a/Engine/Plugins/Experimental/UVEditor/Source/UVEditorTools/Public/UVIslandConformalUnwrapAction.h
+++ b/Engine/Plugins/Experimental/UVEditor/Source/UVEditorTools/Public/UVIslandConformalUnwrapAction.h
@@ -10,7 +10,7 @@
 
 #include "UVIslandConformalUnwrapAction.generated.h"
 
-PREDECLARE_GEOMETRY(class FDynamicMeshSelection);
+PREDECLARE_GEOMETRY(class FUVEditorDynamicMeshSelection);
 
 UCLASS()
 class UVEDITORTOOLS_API UUVIslandConformalUnwrapAction : public UUVToolAction
@@ -24,7 +24,7 @@ public:
 	void SetWorld(UWorld* WorldIn) override;
 	void Setup(UInteractiveTool* ParentToolIn) override;
 	void Shutdown() override;
-	void SetSelection(int32 SelectionTargetIndexIn, const UE::Geometry::FDynamicMeshSelection* NewSelection);
+	void SetSelection(int32 SelectionTargetIndexIn, const UE::Geometry::FUVEditorDynamicMeshSelection* NewSelection);
 
 protected:
 
@@ -33,7 +33,7 @@ protected:
 	int32 MaxIslandSize;
 
 	int32 SelectionTargetIndex;
-	TSharedPtr<UE::Geometry::FDynamicMeshSelection> CurrentSelection;
+	TSharedPtr<UE::Geometry::FUVEditorDynamicMeshSelection> CurrentSelection;
 
 	bool GatherIslandTids();
 	bool PreCheckAction();
diff --git a/Engine/Plugins/Experimental/UVEditor/Source/UVEditorTools/Public/UVSeamSewAction.h b/Engine/Plugins/Experimental/UVEditor/Source/UVEditorTools/Public/UVSeamSewAction.h
--- a/Engine/Plugins/Experimental/UVEditor/Source/UVEditorTools/Public/UVSeamSewAction.h
+++ b/Engine/Plugins/Experimental/UVEditor/Source/UVEditorTools/Public/UVSeamSewAction.h
@@ -10,7 +10,7 @@
 
 #include "UVSeamSewAction.generated.h"
 
-PREDECLARE_GEOMETRY(class FDynamicMeshSelection);
+PREDECLARE_GEOMETRY(class FUVEditorDynamicMeshSelection);
 class APreviewGeometryActor;
 class ULineSetComponent;
 
@@ -45,7 +45,7 @@ public:
 	void Setup(UInteractiveTool* ParentToolIn) override;
 	void Shutdown() override;
 
-	void SetSelection(int32 SelectionTargetIndexIn, const UE::Geometry::FDynamicMeshSelection* NewSelection);
+	void SetSelection(int32 SelectionTargetIndexIn, const UE::Geometry::FUVEditorDynamicMeshSelection* NewSelection);
 	virtual void UpdateVisualizations() override;
 
 protected:
@@ -59,7 +59,7 @@ protected:
 	TObjectPtr<ULineSetComponent> SewEdgePairingLineSet = nullptr;
 
 	int32 SelectionTargetIndex;
-	TSharedPtr<UE::Geometry::FDynamicMeshSelection> CurrentSelection;
+	TSharedPtr<UE::Geometry::FUVEditorDynamicMeshSelection> CurrentSelection;
 
 	int32 FindSewEdgeOppositePairing(int32 UnwrapEid) const;
 	bool PreCheckAction() override;
diff --git a/Engine/Plugins/Experimental/UVEditor/Source/UVEditorTools/Public/UVSelectTool.h b/Engine/Plugins/Experimental/UVEditor/Source/UVEditorTools/Public/UVSelectTool.h
--- a/Engine/Plugins/Experimental/UVEditor/Source/UVEditorTools/Public/UVSelectTool.h
+++ b/Engine/Plugins/Experimental/UVEditor/Source/UVEditorTools/Public/UVSelectTool.h
@@ -16,11 +16,11 @@
 
 #include "UVSelectTool.generated.h"
 
-PREDECLARE_GEOMETRY(class FDynamicMeshSelection);
+PREDECLARE_GEOMETRY(class FUVEditorDynamicMeshSelection);
 class APreviewGeometryActor;
 class FToolCommandChange;
 class ULineSetComponent;
-class UMeshSelectionMechanic;
+class UUVEditorMeshSelectionMechanic;
 class UPointSetComponent;
 class UToolTargetManager;
 class UCombinedTransformGizmo;
@@ -114,7 +114,7 @@ public:
 	}
 
 	// Used by undo/redo changes to update the tool state
-	void SetSelection(const UE::Geometry::FDynamicMeshSelection& NewSelection);
+	void SetSelection(const UE::Geometry::FUVEditorDynamicMeshSelection& NewSelection);
 	void SetGizmoTransform(const FTransform& NewTransform);
 
 	// UInteractiveTool
@@ -157,7 +157,7 @@ protected:
 	TObjectPtr<UUVToolViewportButtonsAPI> ViewportButtonsAPI = nullptr;
 
 	UPROPERTY()
-	TObjectPtr<UMeshSelectionMechanic> SelectionMechanic = nullptr;
+	TObjectPtr<UUVEditorMeshSelectionMechanic> SelectionMechanic = nullptr;
 
 	UPROPERTY()
 	TObjectPtr<UCombinedTransformGizmo> TransformGizmo = nullptr;
diff --git a/Engine/Plugins/Experimental/UVEditor/Source/UVEditorTools/Public/UVToolContextObjects.h b/Engine/Plugins/Experimental/UVEditor/Source/UVEditorTools/Public/UVToolContextObjects.h
--- a/Engine/Plugins/Experimental/UVEditor/Source/UVEditorTools/Public/UVToolContextObjects.h
+++ b/Engine/Plugins/Experimental/UVEditor/Source/UVEditorTools/Public/UVToolContextObjects.h
@@ -7,7 +7,6 @@
 #include "DynamicMesh/DynamicMeshAABBTree3.h"
 #include "DynamicMesh/DynamicMeshChangeTracker.h" // FDynamicMeshChange for TUniquePtr
 #include "InteractiveToolManager.h"
-#include "Selection/DynamicMeshSelection.h"
 #include "GeometryBase.h"
 
 #include "UVToolContextObjects.generated.h"
diff --git a/Engine/Plugins/Runtime/MeshModelingToolset/Source/ModelingComponents/Private/Selection/MeshSelectionMechanic.cpp b/Engine/Plugins/Runtime/MeshModelingToolset/Source/ModelingComponents/Private/Selection/MeshSelectionMechanic.cpp
deleted file mode 100644
diff --git a/Engine/Plugins/Runtime/MeshModelingToolset/Source/ModelingComponents/Public/Selection/DynamicMeshSelection.h b/Engine/Plugins/Runtime/MeshModelingToolset/Source/ModelingComponents/Public/Selection/DynamicMeshSelection.h
deleted file mode 100644
diff --git a/Engine/Plugins/Runtime/MeshModelingToolset/Source/ModelingComponents/Public/Selection/MeshSelectionMechanic.h b/Engine/Plugins/Runtime/MeshModelingToolset/Source/ModelingComponents/Public/Selection/MeshSelectionMechanic.h
deleted file mode 100644
