UVEditor: Enable for Linux/Mac

diff --git a/Engine/Plugins/Experimental/UVEditor/UVEditor.uplugin b/Engine/Plugins/Experimental/UVEditor/UVEditor.uplugin
--- a/Engine/Plugins/Experimental/UVEditor/UVEditor.uplugin
+++ b/Engine/Plugins/Experimental/UVEditor/UVEditor.uplugin
@@ -17,8 +17,7 @@
 		{
 			"Name": "UVEditor",
 			"Type": "Editor",
-			"LoadingPhase": "Default",
-			"PlatformAllowList": [ "Win64" ]
+			"LoadingPhase": "Default"
 		},
 		{
 			"Name": "UVEditorTools",
@@ -28,8 +27,7 @@
 		{
 			"Name": "UVEditorToolsEditorOnly",
 			"Type": "Editor",
-			"LoadingPhase": "Default",
-			"PlatformAllowList": [ "Win64" ]
+			"LoadingPhase": "Default"
 		}
 	],
 	"Plugins": [
