Fix for no unity error
#rnx

diff --git a/Engine/Plugins/Experimental/UVEditor/Source/UVEditorTools/Private/UVEditorUXSettings.cpp b/Engine/Plugins/Experimental/UVEditor/Source/UVEditorTools/Private/UVEditorUXSettings.cpp
--- a/Engine/Plugins/Experimental/UVEditor/Source/UVEditorTools/Private/UVEditorUXSettings.cpp
+++ b/Engine/Plugins/Experimental/UVEditor/Source/UVEditorTools/Private/UVEditorUXSettings.cpp
@@ -1,6 +1,7 @@
 // Copyright Epic Games, Inc. All Rights Reserved.
 
 #include "UVEditorUXSettings.h"
+#include "Math/Color.h"
 
 const float FUVEditorUXSettings::CameraFarPlaneWorldZ(-10.0);
 const float FUVEditorUXSettings::CameraNearPlaneProportionZ(0.8); // Top layer, equivalent to depth bias 80
diff --git a/Engine/Plugins/Experimental/UVEditor/Source/UVEditorTools/Public/UVEditorUXSettings.h b/Engine/Plugins/Experimental/UVEditor/Source/UVEditorTools/Public/UVEditorUXSettings.h
--- a/Engine/Plugins/Experimental/UVEditor/Source/UVEditorTools/Public/UVEditorUXSettings.h
+++ b/Engine/Plugins/Experimental/UVEditor/Source/UVEditorTools/Public/UVEditorUXSettings.h
@@ -3,6 +3,8 @@
 
 #include "CoreMinimal.h"
 
+#include "CoreGlobals.h"
+
 class UVEDITORTOOLS_API FUVEditorUXSettings
 {
 public:
