UVEditor: Add support for drawing rulers in the 2D view and UDIM labels when a UDIM enabled texture is selected as a background image. Additionally enables controls over the grid and ruler rendering in the Display settings.

#rnx

diff --git a/Engine/Plugins/Experimental/UVEditor/Source/UVEditor/Private/UVEditor2DViewportClient.cpp b/Engine/Plugins/Experimental/UVEditor/Source/UVEditor/Private/UVEditor2DViewportClient.cpp
--- a/Engine/Plugins/Experimental/UVEditor/Source/UVEditor/Private/UVEditor2DViewportClient.cpp
+++ b/Engine/Plugins/Experimental/UVEditor/Source/UVEditor/Private/UVEditor2DViewportClient.cpp
@@ -11,12 +11,144 @@
 #include "FrameTypes.h"
 #include "UVEditorMode.h"
 #include "UVEditorUXSettings.h"
+#include "CanvasItem.h"
+#include "CanvasTypes.h"
+#include "Engine/Canvas.h"
+#include "Math/Box.h"
+#include "MathUtil.h"
+#include "UDIMUtilities.h"
+
+namespace FUVEditor2DViewportClientLocals {
+
+	enum class ETextAnchorPosition : uint8
+	{
+		None = 0,
+		VerticalTop = 1 << 0,
+		VerticalMiddle = 1 << 1,
+		VerticalBottom = 1 << 2,
+		HorizontalLeft = 1 << 3,
+		HorizontalMiddle = 1 << 4,
+		HorizontalRight = 1 << 5,
+		TopLeftCorner = VerticalTop | HorizontalLeft,
+		TopRightCorner = VerticalTop | HorizontalRight,
+		BottomLeftCorner = VerticalBottom | HorizontalLeft,
+		BottomRightCorner = VerticalBottom | HorizontalRight,
+		TopCenter = VerticalTop | HorizontalMiddle,
+		BottomCenter = VerticalBottom | HorizontalMiddle,
+		LeftCenter = VerticalMiddle | HorizontalLeft,
+		RightCenter = VerticalMiddle | HorizontalRight,
+		Center = VerticalMiddle | HorizontalMiddle,
+	};
+	ENUM_CLASS_FLAGS(ETextAnchorPosition);
+
+	struct FTextPadding
+	{
+		float TopPadding = 0.0;
+		float BottomPadding = 0.0;
+		float LeftPadding = 0.0;
+		float RightPadding = 0.0;
+	};
+
+	FText ConvertToExponentialNotation(float InValue, float MinExponent = 0.0)
+	{
+		float Exponent = FMath::LogX(10, InValue < 0 ? -InValue : InValue);
+		Exponent = FMath::Floor(Exponent);
+		if (FMath::IsFinite(Exponent) && FMath::Abs(Exponent) > MinExponent)
+		{
+			double Divisor = FMath::Pow(10.0, Exponent);
+			double Base = InValue / Divisor;
+			return FText::Format(FTextFormat::FromString("{0}E{1}"), FText::AsNumber(Base), FText::AsNumber(Exponent));
+		}
+		else
+		{
+			return FText::AsNumber(InValue);
+		}
+	}
+
+	FVector2D ComputeAnchoredPosition(const FVector2D& InPosition, const FText& InText, const UFont* InFont, UCanvas& Canvas, ETextAnchorPosition InAnchorPosition, const FTextPadding& InPadding, bool bClampInsideScreenBounds) {
+		FVector2D OutPosition = InPosition;
+
+		FVector2D TextSize(0, 0), TextShift(0, 0);
+		Canvas.TextSize(InFont, InText.ToString(), TextSize[0], TextSize[1]);
+
+		FBox2D TextBoundingBox(FVector2D(0, 0), TextSize);
+		FBox2D ScreenBoundingBox(FVector2D(0, 0), FVector2D(Canvas.SizeX, Canvas.SizeY));
+
+		if ((InAnchorPosition & ETextAnchorPosition::VerticalTop) != ETextAnchorPosition::None)
+		{
+			TextShift[1] = InPadding.TopPadding; // We are anchored to the top left hand corner by default
+		}
+		else if ((InAnchorPosition & ETextAnchorPosition::VerticalMiddle) != ETextAnchorPosition::None)
+		{
+			TextShift[1] = -TextSize[1] / 2.0;
+		}
+		else if ((InAnchorPosition & ETextAnchorPosition::VerticalBottom) != ETextAnchorPosition::None)
+		{
+			TextShift[1] = -TextSize[1] - InPadding.BottomPadding;
+		}
+
+		if ((InAnchorPosition & ETextAnchorPosition::HorizontalLeft) != ETextAnchorPosition::None)
+		{
+			TextShift[0] = InPadding.LeftPadding; // We are anchored to the top left hand corner by default
+		}
+		else if ((InAnchorPosition & ETextAnchorPosition::HorizontalMiddle) != ETextAnchorPosition::None)
+		{
+			TextShift[0] = -TextSize[0] / 2.0;
+		}
+		else if ((InAnchorPosition & ETextAnchorPosition::HorizontalRight) != ETextAnchorPosition::None)
+		{
+			TextShift[0] = -TextSize[0] - InPadding.RightPadding;
+		}
+		
+		TextBoundingBox = TextBoundingBox.ShiftBy(InPosition + TextShift);
+		
+		if (bClampInsideScreenBounds && !ScreenBoundingBox.IsInside(TextBoundingBox)) {
+			FVector2D TextCenter, TextExtents;
+			TextBoundingBox.GetCenterAndExtents(TextCenter, TextExtents);
+			FBox2D ScreenInsetBoundingBox(ScreenBoundingBox);
+			ScreenInsetBoundingBox = ScreenInsetBoundingBox.ExpandBy(TextExtents*-1.0);
+			FVector2D MovedCenter = ScreenInsetBoundingBox.GetClosestPointTo(TextCenter);
+			TextBoundingBox = TextBoundingBox.MoveTo(MovedCenter);
+		}
+
+		return TextBoundingBox.Min;
+	}
+
+	FCanvasTextItem CreateTextAnchored(const FVector2D& InPosition, const FText& InText, const UFont* InFont,
+		const FLinearColor& InColor, UCanvas& Canvas, ETextAnchorPosition InAnchorPosition, const FTextPadding& InPadding, bool bClampInsideScreenBounds=true)
+	{
+		FVector2D OutPosition = ComputeAnchoredPosition(InPosition, InText, InFont, Canvas, InAnchorPosition, InPadding, bClampInsideScreenBounds);
+		return FCanvasTextItem(OutPosition, InText, InFont, InColor);
+	}
+
+	bool ConvertUVToPixel(const FVector2D& UVIn, FVector2D& PixelOut, const FSceneView& View)
+	{
+		FVector TestWorld = FUVEditorUXSettings::UVToVertPosition(FUVEditorUXSettings::ExternalUVToInternalUV((FVector2f)UVIn));
+		FVector4 TestProjectedHomogenous = View.WorldToScreen(TestWorld);
+		bool bValid = View.ScreenToPixel(TestProjectedHomogenous, PixelOut);
+		return bValid;
+	}
+
+	void ConvertPixelToUV(const FVector2D& PixelIn, double RelDepthIn, FVector2D& UVOut, const FSceneView& View)
+	{
+		FVector4 ScreenPoint = View.PixelToScreen(PixelIn.X, PixelIn.Y, RelDepthIn);
+		FVector4 WorldPointHomogenous = View.ViewMatrices.GetInvViewProjectionMatrix().TransformFVector4(ScreenPoint);
+		FVector WorldPoint(WorldPointHomogenous.X / WorldPointHomogenous.W,
+			               WorldPointHomogenous.Y / WorldPointHomogenous.W,
+			               WorldPointHomogenous.Z / WorldPointHomogenous.W);
+		UVOut = (FVector2D)FUVEditorUXSettings::VertPositionToUV(WorldPoint);
+		UVOut = (FVector2D)FUVEditorUXSettings::InternalUVToExternalUV((FVector2f)UVOut);
+	}
+
+};
+
 
 FUVEditor2DViewportClient::FUVEditor2DViewportClient(FEditorModeTools* InModeTools,
 	FPreviewScene* InPreviewScene, const TWeakPtr<SEditorViewport>& InEditorViewportWidget,
-	UUVToolViewportButtonsAPI* ViewportButtonsAPIIn)
+	UUVToolViewportButtonsAPI* ViewportButtonsAPIIn,
+	UUVTool2DViewportAPI* UVTool2DViewportAPIIn)
 	: FEditorViewportClient(InModeTools, InPreviewScene, InEditorViewportWidget)
-	, ViewportButtonsAPI(ViewportButtonsAPIIn)
+	, ViewportButtonsAPI(ViewportButtonsAPIIn), UVTool2DViewportAPI(UVTool2DViewportAPIIn)
 {
 	ShowWidget(false);
 
@@ -50,6 +182,20 @@ FUVEditor2DViewportClient::FUVEditor2DViewportClient(FEditorModeTools* InModeToo
 	BehaviorSet->Add(ZoomBehavior);
 
 	ModeTools->GetInteractiveToolsContext()->InputRouter->RegisterSource(this);
+
+	static const FName CanvasName(TEXT("UVEditor2DCanvas"));
+	CanvasObject = NewObject<UCanvas>(GetTransientPackage(), CanvasName);
+
+	UVTool2DViewportAPI->OnDrawGridChange.AddLambda(
+		[this](bool bDrawGridIn) {
+			bDrawGrid = bDrawGridIn;
+		});
+
+	UVTool2DViewportAPI->OnDrawRulersChange.AddLambda(
+		[this](bool bDrawRulersIn) {
+			bDrawGridRulers = bDrawRulersIn;
+		});
+
 }
 
 
@@ -65,6 +211,7 @@ void FUVEditor2DViewportClient::AddReferencedObjects(FReferenceCollector& Collec
 
 	Collector.AddReferencedObject(BehaviorSet);
 	Collector.AddReferencedObject(ViewportButtonsAPI);
+	Collector.AddReferencedObject(CanvasObject);
 }
 
 bool FUVEditor2DViewportClient::InputKey(FViewport* InViewport, int32 ControllerId, FKey Key, EInputEvent Event, float/*AmountDepressed*/, bool/*Gamepad*/)
@@ -82,6 +229,17 @@ bool FUVEditor2DViewportClient::InputKey(FViewport* InViewport, int32 Controller
 // Note that this function gets called from the super class Draw(FViewport*, FCanvas*) overload to draw the scene.
 // We don't override that top-level function so that it can do whatever view calculations it needs to do.
 void FUVEditor2DViewportClient::Draw(const FSceneView* View, FPrimitiveDrawInterface* PDI)
+{
+	if (bDrawGrid)
+	{
+		DrawGrid(View, PDI);
+	}
+
+	// Calls ModeTools draw/render functions
+	FEditorViewportClient::Draw(View, PDI);
+}
+
+void FUVEditor2DViewportClient::DrawGrid(const FSceneView* View, FPrimitiveDrawInterface* PDI)
 {
 	// Basic scaling amount
 	const float UVScale = UUVEditorMode::GetUVMeshScalingFactor();
@@ -124,9 +282,6 @@ void FUVEditor2DViewportClient::Draw(const FSceneView* View, FPrimitiveDrawInter
 	PDI->DrawLine(FVector(0, 0, 0), FVector(0, AxisExtent, 0), FUVEditorUXSettings::YAxisColor, SDPG_World, FUVEditorUXSettings::AxisThickness, 0, true);
 
 	// TODO: Draw a little UV axis thing in the lower left, like the XYZ things that normal viewports have.
-
-	// Calls ModeTools draw/render functions
-	FEditorViewportClient::Draw(View, PDI);
 }
 
 bool FUVEditor2DViewportClient::ShouldOrbitCamera() const
@@ -211,4 +366,148 @@ UUVToolViewportButtonsAPI::ESelectionMode FUVEditor2DViewportClient::GetSelectio
 	}
 
 	return ViewportButtonsAPI->GetSelectionMode();
-}
\ No newline at end of file
+}
+
+void FUVEditor2DViewportClient::DrawCanvas(FViewport& InViewport, FSceneView& View, FCanvas& Canvas)
+{
+
+	if (CanvasObject)
+	{
+		CanvasObject->Init(InViewport.GetSizeXY().X, InViewport.GetSizeXY().Y, &View, &Canvas);
+		if (bDrawGridRulers)
+		{
+			DrawGridRulers(InViewport, View, *CanvasObject);
+		}
+
+		bool bEnableUDIMSupport = (FUVEditorUXSettings::CVarEnablePrototypeUDIMSupport.GetValueOnGameThread() > 0);
+		if (bEnableUDIMSupport)
+		{
+			DrawUDIMLabels(InViewport, View, *CanvasObject);
+		}
+	}
+	
+	FEditorViewportClient::DrawCanvas(InViewport, View, Canvas);
+}
+
+void FUVEditor2DViewportClient::DrawGridRulers(FViewport& InViewport, FSceneView& View, UCanvas& Canvas)
+{
+	const float UVScale = FUVEditorUXSettings::UVMeshScalingFactor;
+	const int32 Level = FUVEditorUXSettings::RulerSubdivisionLevel;
+	const int32 Subdivisions = FUVEditorUXSettings::GridSubdivisionsPerLevel; // Number of subdivisions per level
+
+	// Determine important geometry of the viewport for creating grid lines
+	FVector WorldCenterPoint(0, 0, 0);
+	FVector4 WorldToScreenCenter = View.WorldToScreen(WorldCenterPoint);
+	float ZoomFactor = WorldToScreenCenter.W;
+	float GridZoomFactor = ZoomFactor / UVScale;
+
+	float LogZoom = FMath::LogX(Subdivisions, GridZoomFactor);
+	float LogZoomDirection = FMath::Sign(LogZoom);
+	LogZoom = FMath::Abs(LogZoom);
+	LogZoom = FMathf::Floor(LogZoom);
+	LogZoom = LogZoomDirection * LogZoom;
+
+	double GridScale = FMathf::Pow(Subdivisions, LogZoom - Level);
+
+	FVector2D PixelMinBounds(0, Canvas.SizeY);
+	FVector2D PixelMaxBounds(Canvas.SizeX, 0);
+	FVector2D UVMinBounds, UVMaxBounds;
+	FUVEditor2DViewportClientLocals::ConvertPixelToUV(PixelMinBounds, WorldToScreenCenter.Z / WorldToScreenCenter.W, UVMinBounds, View);
+	FUVEditor2DViewportClientLocals::ConvertPixelToUV(PixelMaxBounds, WorldToScreenCenter.Z / WorldToScreenCenter.W, UVMaxBounds, View);
+
+	double GridXOrigin = FMath::GridSnap(UVMinBounds.X, GridScale);
+	double GridYOrigin = FMath::GridSnap(UVMinBounds.Y, GridScale);
+
+	const UFont* Font = GEngine->GetTinyFont();
+
+	double CurrentGridXPos = GridXOrigin;
+	while (CurrentGridXPos < UVMinBounds.X)
+    {
+		CurrentGridXPos += GridScale;
+	}
+
+	do
+    {
+		FVector2D UVRulerPoint(CurrentGridXPos, 0.0);
+		FVector2D PixelRulerPoint;
+		FUVEditor2DViewportClientLocals::ConvertUVToPixel(UVRulerPoint, PixelRulerPoint, View);
+
+		FCanvasTextItem TextItem = FUVEditor2DViewportClientLocals::CreateTextAnchored(PixelRulerPoint,
+			FUVEditor2DViewportClientLocals::ConvertToExponentialNotation(CurrentGridXPos, 2),
+			Font,
+			FLinearColor(FUVEditorUXSettings::RulerXColor),
+			Canvas,
+			FUVEditor2DViewportClientLocals::ETextAnchorPosition::TopLeftCorner,
+			{ 3,0,3,0 });
+
+		Canvas.DrawItem(TextItem);
+		CurrentGridXPos += GridScale;
+	} while (CurrentGridXPos <= UVMaxBounds.X);
+
+	double CurrentGridYPos = GridYOrigin;
+	while (CurrentGridYPos < UVMinBounds.Y)
+    {
+		CurrentGridYPos += GridScale;
+	}
+
+	do
+    {
+		FVector2D UVRulerPoint(0.0, CurrentGridYPos);
+		FVector2D PixelRulerPoint;
+		FUVEditor2DViewportClientLocals::ConvertUVToPixel(UVRulerPoint, PixelRulerPoint, View);
+
+		FCanvasTextItem TextItem = FUVEditor2DViewportClientLocals::CreateTextAnchored(PixelRulerPoint,
+			FUVEditor2DViewportClientLocals::ConvertToExponentialNotation(CurrentGridYPos, 2),
+			Font,
+			FLinearColor(FUVEditorUXSettings::RulerYColor),
+			Canvas,
+			FUVEditor2DViewportClientLocals::ETextAnchorPosition::BottomRightCorner,
+			{ 0,3,0,3 });
+
+		Canvas.DrawItem(TextItem);
+		CurrentGridYPos += GridScale;
+	} while (CurrentGridYPos <= UVMaxBounds.Y);
+
+}
+
+void FUVEditor2DViewportClient::DrawUDIMLabels(FViewport& InViewport, FSceneView& View, UCanvas& Canvas)
+{
+	FVector2D Origin(0,0);
+	FVector2D OneUDimOffet(1.0, 0);
+	FVector2D OriginPixel, OneUDimOffetPixel;
+
+	FUVEditor2DViewportClientLocals::ConvertUVToPixel(Origin, OriginPixel, View);
+	FUVEditor2DViewportClientLocals::ConvertUVToPixel(OneUDimOffet, OneUDimOffetPixel, View);
+
+	double MaxAllowedLabelWidth = FMath::Abs(OriginPixel.X - OneUDimOffetPixel.X);
+
+	const UFont* Font = GEngine->GetLargeFont();
+	FNumberFormattingOptions FormatOptions;
+	FormatOptions.UseGrouping = false;
+
+	for (const FUDIMBlock& Block : UVTool2DViewportAPI->GetUDIMBlocks())
+	{
+		FVector2D TextSize;
+		FText UDIMLabel = FText::AsNumber(UE::TextureUtilitiesCommon::GetUDIMIndex(Block.BlockX, Block.BlockY), &FormatOptions);
+		Canvas.TextSize(Font, UDIMLabel.ToString(), TextSize[0], TextSize[1]);
+		if(TextSize[0] > MaxAllowedLabelWidth)
+		{
+			continue; // Don't draw if our label is bigger than the visual size of one UDIM "block" in the viewport
+		}			
+
+		FVector2D UDimUV(Block.BlockX + 1.0, Block.BlockY + 1.0);
+		FVector2D UDimPixel;
+		FUVEditor2DViewportClientLocals::ConvertUVToPixel(UDimUV, UDimPixel, View);
+
+		FCanvasTextItem TextItem = FUVEditor2DViewportClientLocals::CreateTextAnchored(UDimPixel,
+				UDIMLabel,
+				Font,
+				FLinearColor(FColor(255, 196, 196)),
+				Canvas,
+				FUVEditor2DViewportClientLocals::ETextAnchorPosition::TopRightCorner,
+				{ 5,0,0,5 },
+			    false /* Don't clamp to screen bounds */ );
+
+		Canvas.DrawItem(TextItem);
+	}
+}
diff --git a/Engine/Plugins/Experimental/UVEditor/Source/UVEditor/Private/UVEditorBackgroundPreview.cpp b/Engine/Plugins/Experimental/UVEditor/Source/UVEditor/Private/UVEditorBackgroundPreview.cpp
--- a/Engine/Plugins/Experimental/UVEditor/Source/UVEditor/Private/UVEditorBackgroundPreview.cpp
+++ b/Engine/Plugins/Experimental/UVEditor/Source/UVEditor/Private/UVEditorBackgroundPreview.cpp
@@ -29,10 +29,11 @@ void UUVEditorBackgroundPreview::OnCreated()
 void UUVEditorBackgroundPreview::OnTick(float DeltaTime)
 {
 	if (bSettingsModified)
-	{
+	{		
 		UpdateBackground();
 		UpdateVisibility();
 		bSettingsModified = false;
+		OnBackgroundMaterialChange.Broadcast(BackgroundMaterial);
 	}
 }
 
@@ -41,6 +42,7 @@ void UUVEditorBackgroundPreview::UpdateVisibility()
 	if (Settings->bVisible == false)
 	{
 		BackgroundComponent->SetVisibility(false);
+		ActiveUDIMBlocks.UDIMBlocks.SetNum(0);
 		return;
 	}
 
@@ -50,17 +52,15 @@ void UUVEditorBackgroundPreview::UpdateVisibility()
 
 void UUVEditorBackgroundPreview::UpdateBackground()
 {
+	bool bEnableUDIMSupport = (FUVEditorUXSettings::CVarEnablePrototypeUDIMSupport.GetValueOnGameThread() > 0);
 
-	const int32 GridCellCountX = 1;
-	const int32 GridCellCountY = 1;
-	const FVector Origin = { 0,0,0 };
-	const FVector GridDx = { 1000, 1000, 0 };
 	const FVector Normal(0, 0, 1);
 	const FColor BackgroundColor = FColor::Blue;
 
 	UMaterial* Material = LoadObject<UMaterial>(nullptr, TEXT("/UVEditor/Materials/UVEditorBackground"));
 	check(Material);	
 	BackgroundMaterial = UMaterialInstanceDynamic::Create(Material, this);
+	ActiveUDIMBlocks.UDIMBlocks.SetNum(0);
 	switch (Settings->SourceType)
 	{
 		case EUVEditorBackgroundSourceType::Checkerboard:
@@ -86,6 +86,21 @@ void UUVEditorBackgroundPreview::UpdateBackground()
 				{
 					BackgroundMaterial->SetTextureParameterValue(TEXT("BackgroundVTBaseMap"), Settings->SourceTexture);
 					BackgroundMaterial->SetScalarParameterValue(TEXT("BackgroundVirtualTextureSwitch"), 1);
+
+					// Check for UDIMs
+					if (bEnableUDIMSupport && Settings->SourceTexture->Source.GetNumBlocks() > 1) {
+						ActiveUDIMBlocks.UDIMBlocks.SetNum(Settings->SourceTexture->Source.GetNumBlocks());						
+					
+						for (int32 Block = 0; Block < Settings->SourceTexture->Source.GetNumBlocks(); ++Block) {
+							FTextureSourceBlock SourceBlock;
+							Settings->SourceTexture->Source.GetBlock(Block, SourceBlock);
+
+							ActiveUDIMBlocks.UDIMBlocks[Block].BlockX = SourceBlock.BlockX;
+							ActiveUDIMBlocks.UDIMBlocks[Block].BlockY = SourceBlock.BlockY;
+							ActiveUDIMBlocks.UDIMBlocks[Block].SizeX = SourceBlock.SizeX;
+							ActiveUDIMBlocks.UDIMBlocks[Block].SizeY = SourceBlock.SizeY;
+						}
+					}
 				}
 				else
 				{
@@ -100,30 +115,43 @@ void UUVEditorBackgroundPreview::UpdateBackground()
 			ensure(false);
 	}
 
-	BackgroundMaterial->SetScalarParameterValue(TEXT("BackgroundPixelDepthOffset"), FUVEditorUXSettings::BackgroundQuadDepthOffset);
-	OnBackgroundMaterialChange.Broadcast(BackgroundMaterial);
+	BackgroundMaterial->SetScalarParameterValue(TEXT("BackgroundPixelDepthOffset"), FUVEditorUXSettings::BackgroundQuadDepthOffset);	
 	BackgroundComponent->Clear();
 
-	for (int32 GridStepX = 0; GridStepX < GridCellCountX; ++GridStepX)
+	TArray<FVector2f> UDimBlocksToRender;
+	for (int32 BlockIndex = 0; BlockIndex < ActiveUDIMBlocks.UDIMBlocks.Num(); ++BlockIndex)
 	{
-		for (int32 GridStepY = 0; GridStepY < GridCellCountY; ++GridStepY)
-		{
-			FVector CellOrigin = Origin + FVector(GridDx.X * GridStepX, GridDx.X * GridStepY, 0);
-			FVector CellOffsetX = { GridDx.X, 0, 0 };
-			FVector CellOffsetY = { 0, GridDx.Y, 0 };
-			
-			FRenderableTriangleVertex A(CellOrigin, { -1,0 }, Normal, BackgroundColor);
-			FRenderableTriangleVertex B(CellOrigin + CellOffsetX, { -1 , -1 }, Normal, BackgroundColor);
-			FRenderableTriangleVertex C(CellOrigin + CellOffsetY, { 0, 0 }, Normal, BackgroundColor);
-			FRenderableTriangleVertex D(CellOrigin + CellOffsetX + CellOffsetY, { 0,-1 }, Normal, BackgroundColor);
-
-			FRenderableTriangle Lower(BackgroundMaterial, A, D, B);
-			FRenderableTriangle Upper(BackgroundMaterial, A, C, D);
-
-			BackgroundComponent->AddTriangle(Lower);
-			BackgroundComponent->AddTriangle(Upper);		
-		}
+		UDimBlocksToRender.Push(FVector2f(
+			ActiveUDIMBlocks.UDIMBlocks[BlockIndex].BlockX,
+			ActiveUDIMBlocks.UDIMBlocks[BlockIndex].BlockY
+			));
+	}
+	if (UDimBlocksToRender.Num() == 0)
+	{
+		UDimBlocksToRender.Push(FVector2f(0, 0));
 	}
 
-	
+	for (const FVector2f& GridStep : UDimBlocksToRender)
+	{
+		FVector2f UV00 = { (GridStep.X + 0.0f) , (GridStep.Y + 0.0f) };
+		FVector2f UV01 = { (GridStep.X + 1.0f) , (GridStep.Y + 0.0f) };		
+		FVector2f UV10 = { (GridStep.X + 0.0f) , (GridStep.Y + 1.0f) };
+		FVector2f UV11 = { (GridStep.X + 1.0f) , (GridStep.Y + 1.0f) };
+
+		UV00 = FUVEditorUXSettings::ExternalUVToInternalUV(UV00);
+		UV01 = FUVEditorUXSettings::ExternalUVToInternalUV(UV01);
+		UV10 = FUVEditorUXSettings::ExternalUVToInternalUV(UV10);
+		UV11 = FUVEditorUXSettings::ExternalUVToInternalUV(UV11);
+
+		FRenderableTriangleVertex A((FVector)FUVEditorUXSettings::UVToVertPosition(UV00), (FVector2D)UV00, Normal, BackgroundColor);
+		FRenderableTriangleVertex B((FVector)FUVEditorUXSettings::UVToVertPosition(UV10), (FVector2D)UV10, Normal, BackgroundColor);
+		FRenderableTriangleVertex C((FVector)FUVEditorUXSettings::UVToVertPosition(UV01), (FVector2D)UV01, Normal, BackgroundColor);
+		FRenderableTriangleVertex D((FVector)FUVEditorUXSettings::UVToVertPosition(UV11), (FVector2D)UV11, Normal, BackgroundColor);
+
+		FRenderableTriangle Lower(BackgroundMaterial, A, D, B);
+		FRenderableTriangle Upper(BackgroundMaterial, A, C, D);
+
+		BackgroundComponent->AddTriangle(Lower);
+		BackgroundComponent->AddTriangle(Upper);
+	}
 }
diff --git a/Engine/Plugins/Experimental/UVEditor/Source/UVEditor/Private/UVEditorMode.cpp b/Engine/Plugins/Experimental/UVEditor/Source/UVEditor/Private/UVEditorMode.cpp
--- a/Engine/Plugins/Experimental/UVEditor/Source/UVEditor/Private/UVEditorMode.cpp
+++ b/Engine/Plugins/Experimental/UVEditor/Source/UVEditor/Private/UVEditorMode.cpp
@@ -143,6 +143,10 @@ UUVEditorMode::UUVEditorMode()
 		false);
 }
 
+double UUVEditorMode::GetUVMeshScalingFactor() {
+	return FUVEditorUXSettings::UVMeshScalingFactor;
+}
+
 void UUVEditorMode::Enter()
 {
 	Super::Enter();
@@ -151,15 +155,43 @@ void UUVEditorMode::Enter()
 	BackgroundVisualization->CreateInWorld(GetWorld(), FTransform::Identity);
 
 	BackgroundVisualization->Settings->WatchProperty(BackgroundVisualization->Settings->bVisible, 
-		[this](bool IsVisible) { UpdateTriangleMaterialBasedOnBackground(IsVisible); });
+		[this](bool IsVisible) {
+			UpdateTriangleMaterialBasedOnBackground(IsVisible);
+			UpdateViewportUDIMLabels(IsVisible);			
+		});
 
 	BackgroundVisualization->OnBackgroundMaterialChange.AddWeakLambda(this,
 		[this](TObjectPtr<UMaterialInstanceDynamic> MaterialInstance) {
 			UpdatePreviewMaterialBasedOnBackground();
+			UpdateViewportUDIMLabels(true);
 		});
 
 	PropertyObjectsToTick.Add(BackgroundVisualization->Settings);
 
+	UVEditorGridProperties = NewObject <UUVEditorGridProperties>(this);
+
+	UVEditorGridProperties->WatchProperty(UVEditorGridProperties->bDrawGrid,
+		[this](bool bDrawGrid) {
+			UContextObjectStore* ContextStore = GetInteractiveToolsContext()->ToolManager->GetContextObjectStore();
+			UUVTool2DViewportAPI* UVTool2DViewportAPI = ContextStore->FindContext<UUVTool2DViewportAPI>();
+			if (UVTool2DViewportAPI)
+			{
+				UVTool2DViewportAPI->SetDrawGrid(bDrawGrid, true);
+			}		
+		});
+
+	UVEditorGridProperties->WatchProperty(UVEditorGridProperties->bDrawRulers,
+		[this](bool bDrawRulers) {
+			UContextObjectStore* ContextStore = GetInteractiveToolsContext()->ToolManager->GetContextObjectStore();
+			UUVTool2DViewportAPI* UVTool2DViewportAPI = ContextStore->FindContext<UUVTool2DViewportAPI>();
+			if (UVTool2DViewportAPI)
+			{
+				UVTool2DViewportAPI->SetDrawRulers(bDrawRulers, true);
+			}
+		});
+
+	PropertyObjectsToTick.Add(UVEditorGridProperties);
+
 	RegisterTools();
 
 	ActivateDefaultTool();
@@ -236,6 +268,16 @@ UObject* UUVEditorMode::GetBackgroundSettingsObject()
 	return nullptr;
 }
 
+UObject* UUVEditorMode::GetGridSettingsObject()
+{
+	if (UVEditorGridProperties)
+	{
+		return UVEditorGridProperties;
+	}
+	return nullptr;
+}
+
+
 void UUVEditorMode::ActivateDefaultTool()
 {
 	GetInteractiveToolsContext()->StartTool(DefaultToolIdentifier);
@@ -376,7 +418,7 @@ void UUVEditorMode::Exit()
 }
 
 void UUVEditorMode::InitializeContexts(FEditorViewportClient& LivePreviewViewportClient, FAssetEditorModeManager& LivePreviewModeManager, 
-	UUVToolViewportButtonsAPI& ViewportButtonsAPI)
+	UUVToolViewportButtonsAPI& ViewportButtonsAPI, UUVTool2DViewportAPI& UVTool2DViewportAPI)
 {
 	using namespace UVEditorModeLocals;
 
@@ -416,6 +458,8 @@ void UUVEditorMode::InitializeContexts(FEditorViewportClient& LivePreviewViewpor
 	ContextStore->AddContextObject(AssetAndLayerAPI);
 
 	ContextStore->AddContextObject(&ViewportButtonsAPI);
+
+	ContextStore->AddContextObject(&UVTool2DViewportAPI);
 }
 
 void UUVEditorMode::InitializeTargets(const TArray<TObjectPtr<UObject>>& AssetsIn,
@@ -467,15 +511,6 @@ void UUVEditorMode::InitializeTargets(const TArray<TObjectPtr<UObject>>& AssetsI
 	// frequently end up negative).
 	// The ScaleFactor just scales the mesh up. Scaling the mesh up makes it easier to zoom in
 	// further into the display before getting issues with the camera near plane distance.
-	double ScaleFactor = GetUVMeshScalingFactor();
-	auto UVToVertPosition = [this, ScaleFactor](const FVector2f& UV)
-	{
-		return FVector3d((1 - UV.Y) * ScaleFactor, UV.X * ScaleFactor, 0);
-	};
-	auto VertPositionToUV = [this, ScaleFactor](const FVector3d& VertPosition)
-	{
-		return FVector2f(VertPosition.Y / ScaleFactor, 1 - (VertPosition.X / ScaleFactor));
-	};
 
 	// Construct the full input objects that the tools actually operate on.
 	for (int32 AssetID = 0; AssetID < ToolTargets.Num(); ++AssetID)
@@ -485,7 +520,7 @@ void UUVEditorMode::InitializeTargets(const TArray<TObjectPtr<UObject>>& AssetsI
 		if (!ToolInputObject->InitializeMeshes(ToolTargets[AssetID], AppliedCanonicalMeshes[AssetID],
 			AppliedPreviews[AssetID], AssetID, DefaultUVLayerIndex,
 			GetWorld(), LivePreviewWorld, ToolSetupUtil::GetDefaultWorkingMaterial(GetToolManager()),
-			UVToVertPosition, VertPositionToUV))
+			FUVEditorUXSettings::UVToVertPosition, FUVEditorUXSettings::VertPositionToUV))
 		{
 			return;
 		}
@@ -581,6 +616,35 @@ void UUVEditorMode::ApplyChanges()
 	GetToolManager()->EndUndoTransaction();
 }
 
+void UUVEditorMode::UpdateViewportUDIMLabels(bool IsVisible)
+{
+	bool bEnableUDIMSupport = (FUVEditorUXSettings::CVarEnablePrototypeUDIMSupport.GetValueOnGameThread() > 0);
+	if (!bEnableUDIMSupport)
+	{
+		return;
+	}
+
+	UContextObjectStore* ContextStore = GetInteractiveToolsContext()->ToolManager->GetContextObjectStore();
+	UUVTool2DViewportAPI* UVTool2DViewportAPI = ContextStore->FindContext<UUVTool2DViewportAPI>();
+	if (UVTool2DViewportAPI)
+	{
+		TArray<FUDIMBlock> Blocks;
+		if (IsVisible)
+		{
+			Blocks.SetNum(BackgroundVisualization->ActiveUDIMBlocks.UDIMBlocks.Num());
+			for (int32 BlockIndex = 0; BlockIndex < BackgroundVisualization->ActiveUDIMBlocks.UDIMBlocks.Num(); ++BlockIndex)
+			{
+				Blocks[BlockIndex].BlockX = BackgroundVisualization->ActiveUDIMBlocks.UDIMBlocks[BlockIndex].BlockX;
+				Blocks[BlockIndex].BlockY = BackgroundVisualization->ActiveUDIMBlocks.UDIMBlocks[BlockIndex].BlockY;
+				Blocks[BlockIndex].SizeX = BackgroundVisualization->ActiveUDIMBlocks.UDIMBlocks[BlockIndex].SizeX;
+				Blocks[BlockIndex].SizeY = BackgroundVisualization->ActiveUDIMBlocks.UDIMBlocks[BlockIndex].SizeY;
+			}
+		}
+		UVTool2DViewportAPI->SetUDIMBlocks(Blocks, true);
+	}
+}
+
+
 void UUVEditorMode::UpdateTriangleMaterialBasedOnBackground(bool IsBackgroundVisible)
 {
 	float TriangleOpacity;
diff --git a/Engine/Plugins/Experimental/UVEditor/Source/UVEditor/Private/UVEditorModeToolkit.cpp b/Engine/Plugins/Experimental/UVEditor/Source/UVEditor/Private/UVEditorModeToolkit.cpp
--- a/Engine/Plugins/Experimental/UVEditor/Source/UVEditor/Private/UVEditorModeToolkit.cpp
+++ b/Engine/Plugins/Experimental/UVEditor/Source/UVEditor/Private/UVEditorModeToolkit.cpp
@@ -289,7 +289,7 @@ TSharedRef<SWidget> FUVEditorModeToolkit::CreateBackgroundSettingsWidget()
 {
 	TSharedRef<SBorder> BackgroundDetailsContainer =
 		SNew(SBorder)
-		.BorderImage(FEditorStyle::GetBrush("NoBorder"));
+		.BorderImage(FAppStyle::Get().GetBrush("NoBorder"));
 
 	TSharedRef<SWidget> Widget = SNew(SBorder)
 		.HAlign(HAlign_Fill)
@@ -320,6 +320,41 @@ TSharedRef<SWidget> FUVEditorModeToolkit::CreateBackgroundSettingsWidget()
 	return Widget;
 }
 
+TSharedRef<SWidget> FUVEditorModeToolkit::CreateGridSettingsWidget()
+{
+	TSharedRef<SBorder> GridDetailsContainer =
+		SNew(SBorder)
+		.BorderImage(FAppStyle::Get().GetBrush("NoBorder"));
+
+	TSharedRef<SWidget> Widget = SNew(SBorder)
+		.HAlign(HAlign_Fill)
+		.Padding(4)
+		[
+			SNew(SBox)
+			.MinDesiredWidth(500)
+		[
+			GridDetailsContainer
+		]
+		];
+
+	FPropertyEditorModule& PropertyEditorModule = FModuleManager::GetModuleChecked<FPropertyEditorModule>("PropertyEditor");
+
+	FDetailsViewArgs GridDetailsViewArgs;
+	GridDetailsViewArgs.bAllowSearch = false;
+	GridDetailsViewArgs.NameAreaSettings = FDetailsViewArgs::HideNameArea;
+	GridDetailsViewArgs.bHideSelectionTip = true;
+	GridDetailsViewArgs.DefaultsOnlyVisibility = EEditDefaultsOnlyNodeVisibility::Automatic;
+	GridDetailsViewArgs.bShowOptions = false;
+	GridDetailsViewArgs.bAllowMultipleTopLevelObjects = false;
+
+	UUVEditorMode* Mode = Cast<UUVEditorMode>(GetScriptableEditorMode());
+	TSharedRef<IDetailsView> GridDetailsView = PropertyEditorModule.CreateDetailView(GridDetailsViewArgs);
+	GridDetailsView->SetObject(Mode->GetGridSettingsObject());
+	GridDetailsContainer->SetContent(GridDetailsView);
+
+	return Widget;
+}
+
 void FUVEditorModeToolkit::UpdateActiveToolProperties()
 {
 	UInteractiveTool* CurTool = GetScriptableEditorMode()->GetToolManager()->GetActiveTool(EToolSide::Left);
diff --git a/Engine/Plugins/Experimental/UVEditor/Source/UVEditor/Private/UVEditorToolkit.cpp b/Engine/Plugins/Experimental/UVEditor/Source/UVEditor/Private/UVEditorToolkit.cpp
--- a/Engine/Plugins/Experimental/UVEditor/Source/UVEditor/Private/UVEditorToolkit.cpp
+++ b/Engine/Plugins/Experimental/UVEditor/Source/UVEditor/Private/UVEditorToolkit.cpp
@@ -27,6 +27,7 @@
 #include "UVToolContextObjects.h"
 #include "UVEditorModeUILayer.h"
 #include "Widgets/Docking/SDockTab.h"
+#include "UVEditorUXSettings.h"
 
 #include "SLevelViewport.h"
 
@@ -117,6 +118,8 @@ FUVEditorToolkit::FUVEditorToolkit(UAssetEditor* InOwningAssetEditor)
 	// We create it here so that we can pass it both into the 2d viewport and when we initialize 
 	// the mode.
 	ViewportButtonsAPI = NewObject<UUVToolViewportButtonsAPI>();
+
+	UVTool2DViewportAPI = NewObject<UUVTool2DViewportAPI>();
 }
 
 FUVEditorToolkit::~FUVEditorToolkit()
@@ -341,7 +344,7 @@ TSharedPtr<FEditorViewportClient> FUVEditorToolkit::CreateEditorViewportClient()
 	// Instead, we do viewport client adjustment in PostInitAssetEditor().
 	check(EditorModeManager.IsValid());
 	return MakeShared<FUVEditor2DViewportClient>(EditorModeManager.Get(), UnwrapScene.Get(), 
-		UVEditor2DViewport, ViewportButtonsAPI);
+		UVEditor2DViewport, ViewportButtonsAPI, UVTool2DViewportAPI);
 }
 
 // Called from FBaseAssetToolkit::CreateWidgets. The delegate call path goes through FAssetEditorToolkit::InitAssetEditor
@@ -389,7 +392,7 @@ void FUVEditorToolkit::PostInitAssetEditor()
 	check(UVMode);
 
 	// The mode will need to be able to get to the live preview world, camera, input router, and viewport buttons.
-	UVMode->InitializeContexts(*LivePreviewViewportClient, *LivePreviewEditorModeManager, *ViewportButtonsAPI);
+	UVMode->InitializeContexts(*LivePreviewViewportClient, *LivePreviewEditorModeManager, *ViewportButtonsAPI, *UVTool2DViewportAPI);
 
 	TArray<TObjectPtr<UObject>> ObjectsToEdit;
 	OwningAssetEditor->GetObjectsToEdit(ObjectsToEdit);
@@ -447,7 +450,39 @@ void FUVEditorToolkit::PostInitAssetEditor()
 		FUIAction(),
 		FOnGetContent::CreateLambda([UVModeToolkit]()
 		{
-			return UVModeToolkit->CreateBackgroundSettingsWidget();
+			
+			TSharedRef<SVerticalBox> Container = SNew(SVerticalBox);
+
+			Container->AddSlot()
+				.AutoHeight()
+				.Padding(FMargin(0.f, 0.f, 8.f, 0.f))
+				[
+					SNew(SBox)
+					.MinDesiredWidth(500)
+				[
+					UVModeToolkit->CreateGridSettingsWidget()
+				]
+				];
+
+			Container->AddSlot()
+				.AutoHeight()
+				.Padding(FMargin(0.f, 0.f, 8.f, 0.f))
+				[
+					SNew(SBox)
+					.MinDesiredWidth(500)
+				[
+					UVModeToolkit->CreateBackgroundSettingsWidget()
+				]
+				];
+
+			TSharedRef<SWidget> Widget = SNew(SBorder)
+				.HAlign(HAlign_Fill)
+				.Padding(4)
+				[
+					Container
+				];
+
+			return Widget;
 		}),
 		LOCTEXT("UVEditorBackgroundSettings_Label", "Display"),
 		LOCTEXT("UVEditorBackgroundSettings_ToolTip", "Change the background display settings"),
@@ -489,7 +524,7 @@ void FUVEditorToolkit::PostInitAssetEditor()
 	UUVEditorSubsystem* UVSubsystem = GEditor->GetEditorSubsystem<UUVEditorSubsystem>();
 	if (UVSubsystem)
 	{
-		ScaleFactor = UUVEditorMode::GetUVMeshScalingFactor();
+		ScaleFactor = FUVEditorUXSettings::UVMeshScalingFactor;
 	}
 	ViewportClient->SetViewLocation(FVector(ScaleFactor / 2, ScaleFactor / 2, ScaleFactor));
 	ViewportClient->SetViewRotation(FRotator(-90, 0, 0));
diff --git a/Engine/Plugins/Experimental/UVEditor/Source/UVEditor/Public/UVEditor2DViewportClient.h b/Engine/Plugins/Experimental/UVEditor/Source/UVEditor/Public/UVEditor2DViewportClient.h
--- a/Engine/Plugins/Experimental/UVEditor/Source/UVEditor/Public/UVEditor2DViewportClient.h
+++ b/Engine/Plugins/Experimental/UVEditor/Source/UVEditor/Public/UVEditor2DViewportClient.h
@@ -10,6 +10,7 @@
 #include "UVToolContextObjects.h" // UUVToolViewportButtonsAPI::ESelectionMode
 
 class UUVToolViewportButtonsAPI;
+class UCanvas;
 
 /**
  * Client used to display a 2D view of the UV's, implemented by using a perspective viewport with a locked
@@ -19,14 +20,13 @@ class UVEDITOR_API FUVEditor2DViewportClient : public FEditorViewportClient, pub
 {
 public:
 	FUVEditor2DViewportClient(FEditorModeTools* InModeTools, FPreviewScene* InPreviewScene,
-		const TWeakPtr<SEditorViewport>& InEditorViewportWidget, UUVToolViewportButtonsAPI* ViewportButtonsAPI);
+		const TWeakPtr<SEditorViewport>& InEditorViewportWidget, UUVToolViewportButtonsAPI* ViewportButtonsAPI, UUVTool2DViewportAPI* UVTool2DViewportAPI);
 
 	virtual ~FUVEditor2DViewportClient() {}
 
 	bool AreSelectionButtonsEnabled() const;
 	void SetSelectionMode(UUVToolViewportButtonsAPI::ESelectionMode NewMode);
 	UUVToolViewportButtonsAPI::ESelectionMode GetSelectionMode() const;
-
 	bool AreWidgetButtonsEnabled() const;
 
 	// FEditorViewportClient
@@ -35,7 +35,10 @@ public:
 	virtual bool ShouldOrbitCamera() const override;
 	bool CanSetWidgetMode(UE::Widget::EWidgetMode NewMode) const override;
 	void SetWidgetMode(UE::Widget::EWidgetMode NewMode) override;
-	UE::Widget::EWidgetMode GetWidgetMode() const override;
+	UE::Widget::EWidgetMode GetWidgetMode() const override;	
+	void DrawCanvas(FViewport& InViewport, FSceneView& View, FCanvas& Canvas) override;
+
+
 	// Overriding base class visibility
 	using FEditorViewportClient::OverrideNearClipPlane;
 
@@ -46,9 +49,18 @@ public:
 	virtual void AddReferencedObjects(FReferenceCollector& Collector) override;
 
 protected:
+	void DrawGrid(const FSceneView* View, FPrimitiveDrawInterface* PDI);
+	void DrawGridRulers(FViewport& InViewport, FSceneView& View, UCanvas& Canvas);
+	void DrawUDIMLabels(FViewport& InViewport, FSceneView& View, UCanvas& Canvas);
+
 	// These get added in AddReferencedObjects for memory management
 	UInputBehaviorSet* BehaviorSet;
 	UUVToolViewportButtonsAPI* ViewportButtonsAPI;
+	UUVTool2DViewportAPI* UVTool2DViewportAPI;
+
+	bool bDrawGridRulers = true;
+	bool bDrawGrid = true;
+	UCanvas* CanvasObject;
 
 	// Note that it's generally less hassle if the unique ptr types are complete here,
 	// not forward declared, else we get compile errors if their destruction shows up
diff --git a/Engine/Plugins/Experimental/UVEditor/Source/UVEditor/Public/UVEditorBackgroundPreview.h b/Engine/Plugins/Experimental/UVEditor/Source/UVEditor/Public/UVEditorBackgroundPreview.h
--- a/Engine/Plugins/Experimental/UVEditor/Source/UVEditor/Public/UVEditorBackgroundPreview.h
+++ b/Engine/Plugins/Experimental/UVEditor/Source/UVEditor/Public/UVEditorBackgroundPreview.h
@@ -21,6 +21,33 @@ enum class EUVEditorBackgroundSourceType
 	Material
 };
 
+USTRUCT()
+struct UVEDITOR_API FUVEditorBackgroundUDIMBlockInfo
+{
+	GENERATED_BODY()
+
+	UPROPERTY()
+	int32 BlockX;
+
+	UPROPERTY()
+	int32 BlockY;
+
+	UPROPERTY()
+	int32 SizeX;
+
+	UPROPERTY()
+	int32 SizeY;
+};
+
+USTRUCT()
+struct UVEDITOR_API FUVEditorBackgroundUDIMBlocks
+{
+	GENERATED_BODY()
+
+	UPROPERTY()
+	TArray<FUVEditorBackgroundUDIMBlockInfo> UDIMBlocks;
+};
+
 /**
  * Visualization settings for the UUVEditorBackgroundPreview
  */
@@ -68,6 +95,9 @@ public:
 	UPROPERTY()
 	TObjectPtr<UUVEditorBackgroundPreviewProperties> Settings;
 
+	UPROPERTY()
+	FUVEditorBackgroundUDIMBlocks ActiveUDIMBlocks;
+
 	/** The component containing the quad visualization */
 	UPROPERTY()
 	TObjectPtr<UTriangleSetComponent> BackgroundComponent;
diff --git a/Engine/Plugins/Experimental/UVEditor/Source/UVEditor/Public/UVEditorMode.h b/Engine/Plugins/Experimental/UVEditor/Source/UVEditor/Public/UVEditorMode.h
--- a/Engine/Plugins/Experimental/UVEditor/Source/UVEditor/Public/UVEditorMode.h
+++ b/Engine/Plugins/Experimental/UVEditor/Source/UVEditor/Public/UVEditorMode.h
@@ -7,6 +7,7 @@
 #include "Tools/UEdMode.h"
 #include "ToolTargets/ToolTarget.h"
 #include "GeometryBase.h"
+#include "InteractiveTool.h"
 
 #include "UVEditorMode.generated.h"
 
@@ -27,8 +28,26 @@ class UUVToolStateObjectStore;
 class UUVEditorToolMeshInput;
 class UUVEditorBackgroundPreview;
 class UUVToolViewportButtonsAPI;
+class UUVTool2DViewportAPI;
 
 
+/**
+ * Visualization settings for the UUVEditorMode's Grid
+ */
+UCLASS()
+class UVEDITOR_API UUVEditorGridProperties : public UInteractiveToolPropertySet
+{
+	GENERATED_BODY()
+public:
+	/** Should the grid be shown?*/
+	UPROPERTY(EditAnywhere, Category = Grid, meta = (DisplayName = "Display Grid"))
+	bool bDrawGrid = true;
+
+	/** Should the grid rulers be shown?*/
+	UPROPERTY(EditAnywhere, Category = Grid, meta = (DisplayName = "Display Rulers"))
+	bool bDrawRulers = true;
+};
+
 /**
  * The UV editor mode is the mode used in the UV asset editor. It holds most of the inter-tool state.
  * We put things in a mode instead of directly into the asset editor in case we want to someday use the mode
@@ -54,12 +73,12 @@ public:
 	/**
 	 * Gets the factor by which UV layer unwraps get scaled (scaling makes certain things easier, like zooming in, etc).
 	 */
-	static double GetUVMeshScalingFactor() { return 1000; }
+	static double GetUVMeshScalingFactor();
 
 	// Both initialization functions must be called for things to function properly. InitializeContexts should
 	// be done first so that the 3d preview world is ready for creating meshes in InitializeTargets.
 	void InitializeContexts(FEditorViewportClient& LivePreviewViewportClient, FAssetEditorModeManager& LivePreviewModeManager,
-		UUVToolViewportButtonsAPI& ViewportButtonsAPI);
+		UUVToolViewportButtonsAPI& ViewportButtonsAPI, UUVTool2DViewportAPI& UVTool2DViewportAPI);
 	void InitializeTargets(const TArray<TObjectPtr<UObject>>& AssetsIn, const TArray<FTransform>& TransformsIn);
 
 	// Public for use by undo/redo. Otherwise should use RequestUVChannelChange
@@ -101,6 +120,9 @@ public:
 	/** @return A settings object suitable for display in a details panel to control the background visualization. */
 	UObject* GetBackgroundSettingsObject();
 
+	/** @return A settings object suitable for display in a details panel to control the grid. */
+	UObject* GetGridSettingsObject();
+
 	// UEdMode overrides
 	virtual void Enter() override;
 	virtual bool ShouldToolStartBeAllowed(const FString& ToolIdentifier) const override;
@@ -122,6 +144,9 @@ public:
 	UPROPERTY()
 	TObjectPtr<UUVEditorBackgroundPreview> BackgroundVisualization;
 
+	// Hold a settings object to configure the grid
+	UPROPERTY()
+	TObjectPtr<UUVEditorGridProperties> UVEditorGridProperties = nullptr;
 protected:
 
 	// UEdMode overrides
@@ -133,6 +158,7 @@ protected:
 	
 	void UpdateTriangleMaterialBasedOnBackground(bool IsBackgroundVisible);
 	void UpdatePreviewMaterialBasedOnBackground();
+	void UpdateViewportUDIMLabels(bool IsVisible);
 
 	/**
 	 * Stores original input objects, for instance UStaticMesh pointers. AssetIDs on tool input 
diff --git a/Engine/Plugins/Experimental/UVEditor/Source/UVEditor/Public/UVEditorModeToolkit.h b/Engine/Plugins/Experimental/UVEditor/Source/UVEditor/Public/UVEditorModeToolkit.h
--- a/Engine/Plugins/Experimental/UVEditor/Source/UVEditor/Public/UVEditorModeToolkit.h
+++ b/Engine/Plugins/Experimental/UVEditor/Source/UVEditor/Public/UVEditorModeToolkit.h
@@ -29,6 +29,9 @@ public:
 	/** Creates a widget where the background visualization can be changed. */
 	virtual TSharedRef<SWidget> CreateBackgroundSettingsWidget();
 
+	/** Creates a widget where the grid display can be changed. */
+	virtual TSharedRef<SWidget> CreateGridSettingsWidget();
+
 	// FModeToolkit
 	virtual void Init(const TSharedPtr<IToolkitHost>& InitToolkitHost, TWeakObjectPtr<UEdMode> InOwningMode) override;
 	virtual FText GetActiveToolDisplayName() const override { return ActiveToolName; }
diff --git a/Engine/Plugins/Experimental/UVEditor/Source/UVEditor/Public/UVEditorToolkit.h b/Engine/Plugins/Experimental/UVEditor/Source/UVEditor/Public/UVEditorToolkit.h
--- a/Engine/Plugins/Experimental/UVEditor/Source/UVEditor/Public/UVEditorToolkit.h
+++ b/Engine/Plugins/Experimental/UVEditor/Source/UVEditor/Public/UVEditorToolkit.h
@@ -16,6 +16,7 @@ class SWidget;
 class UInteractiveToolsContext;
 class UInputRouter;
 class UUVToolViewportButtonsAPI;
+class UUVTool2DViewportAPI;
 
 /**
  * The toolkit is supposed to act as the UI manager for the asset editor. It's responsible 
@@ -90,6 +91,8 @@ protected:
 	TWeakPtr<SEditorViewport> UVEditor2DViewport;
 	UUVToolViewportButtonsAPI* ViewportButtonsAPI = nullptr;
 
+	UUVTool2DViewportAPI* UVTool2DViewportAPI = nullptr;
+
 	TSharedPtr<FUVEditorModeUILayer> ModeUILayer;
 	TSharedPtr<FWorkspaceItem> UVEditorMenuCategory;
 };
\ No newline at end of file
diff --git a/Engine/Plugins/Experimental/UVEditor/Source/UVEditor/UVEditor.Build.cs b/Engine/Plugins/Experimental/UVEditor/Source/UVEditor/UVEditor.Build.cs
--- a/Engine/Plugins/Experimental/UVEditor/Source/UVEditor/UVEditor.Build.cs
+++ b/Engine/Plugins/Experimental/UVEditor/Source/UVEditor/UVEditor.Build.cs
@@ -62,6 +62,7 @@ public class UVEditor : ModuleRules
 				"UVEditorTools",
 				"UVEditorToolsEditorOnly",
 				"WorkspaceMenuStructure",
+				"TextureUtilitiesCommon"
 				// ... add private dependencies that you statically link with here ...	
 			}
 			);
diff --git a/Engine/Plugins/Experimental/UVEditor/Source/UVEditorTools/Private/UVEditorUXSettings.cpp b/Engine/Plugins/Experimental/UVEditor/Source/UVEditorTools/Private/UVEditorUXSettings.cpp
--- a/Engine/Plugins/Experimental/UVEditor/Source/UVEditorTools/Private/UVEditorUXSettings.cpp
+++ b/Engine/Plugins/Experimental/UVEditor/Source/UVEditorTools/Private/UVEditorUXSettings.cpp
@@ -3,6 +3,7 @@
 #include "UVEditorUXSettings.h"
 #include "Math/Color.h"
 
+const float FUVEditorUXSettings::UVMeshScalingFactor(1000.0);
 const float FUVEditorUXSettings::CameraFarPlaneWorldZ(-10.0);
 const float FUVEditorUXSettings::CameraNearPlaneProportionZ(0.8); // Top layer, equivalent to depth bias 80
 
@@ -62,7 +63,8 @@ const FColor FUVEditorUXSettings::XAxisColor(FColor::Red);
 const FColor FUVEditorUXSettings::YAxisColor(FColor::Green);
 const FColor FUVEditorUXSettings::GridMajorColor(FColor::FromHex("#888888"));
 const FColor FUVEditorUXSettings::GridMinorColor(FColor::FromHex("#777777"));
-
+const FColor FUVEditorUXSettings::RulerXColor(FColor::FromHex("#888888"));
+const FColor FUVEditorUXSettings::RulerYColor(FColor::FromHex("#888888"));
 // Thicknesses
 const float FUVEditorUXSettings::LivePreviewHighlightThickness(2.0);
 const float FUVEditorUXSettings::LivePreviewHighlightPointSize(4);
@@ -80,6 +82,14 @@ const float FUVEditorUXSettings::ToolPointSize(6);
 // Grid
 const int32 FUVEditorUXSettings::GridSubdivisionsPerLevel(4);
 const int32 FUVEditorUXSettings::GridLevels(3);
+const int32 FUVEditorUXSettings::RulerSubdivisionLevel(1);
+
+// CVARs
+
+TAutoConsoleVariable<int32> FUVEditorUXSettings::CVarEnablePrototypeUDIMSupport(
+	TEXT("modeling.UVEditor.UDIMSupport"),
+	0,
+	TEXT("Enable experimental UDIM support in the UVEditor"));
 
 
 FLinearColor FUVEditorUXSettings::GetTriangleColorByTargetIndex(int32 TargetIndex)
@@ -106,3 +116,23 @@ FLinearColor FUVEditorUXSettings::GetBoundaryColorByTargetIndex(int32 TargetInde
 	BoundaryColorHSV.B = UnwrapBoundaryValue;
 	return BoundaryColorHSV.HSVToLinearRGB();
 }
+
+FVector2f FUVEditorUXSettings::ExternalUVToInternalUV(const FVector2f& UV)
+{
+	return FVector2f(UV.X, 1-UV.Y);
+}
+
+FVector2f FUVEditorUXSettings::InternalUVToExternalUV(const FVector2f& UV)
+{
+	return FVector2f(UV.X, 1-UV.Y);
+}
+
+FVector3d FUVEditorUXSettings::UVToVertPosition(const FVector2f& UV)
+{
+	return FVector3d((1 - UV.Y) * UVMeshScalingFactor, UV.X * UVMeshScalingFactor, 0);
+}
+
+FVector2f FUVEditorUXSettings::VertPositionToUV(const FVector3d& VertPosition)
+{
+	return FVector2f(VertPosition.Y / UVMeshScalingFactor, 1 - (VertPosition.X / UVMeshScalingFactor));
+};
\ No newline at end of file
diff --git a/Engine/Plugins/Experimental/UVEditor/Source/UVEditorTools/Public/UVEditorUXSettings.h b/Engine/Plugins/Experimental/UVEditor/Source/UVEditorTools/Public/UVEditorUXSettings.h
--- a/Engine/Plugins/Experimental/UVEditor/Source/UVEditorTools/Public/UVEditorUXSettings.h
+++ b/Engine/Plugins/Experimental/UVEditor/Source/UVEditorTools/Public/UVEditorUXSettings.h
@@ -13,8 +13,12 @@ public:
 
 	// General Display Properties
 
+	// The ScaleFactor just scales the mesh up. Scaling the mesh up makes it easier to zoom in
+	// further into the display before getting issues with the camera near plane distance.
+	static const float UVMeshScalingFactor;
+
 	// Position to place the 2D camera far plane relative to world z
-	static const float CameraFarPlaneWorldZ;;
+	static const float CameraFarPlaneWorldZ;
 
 	// The near plane gets positioned some proportion to z = 0. We don't use a constant value because our depth offset values are percentage-based
 	// Lower proportion move the plane nearer to world z
@@ -85,15 +89,48 @@ public:
 	static const float GridMajorThickness;
 	static const FColor XAxisColor;
 	static const FColor YAxisColor;
+	static const FColor RulerXColor;
+	static const FColor RulerYColor;
 	static const FColor GridMajorColor;
 	static const FColor GridMinorColor;
 	static const int32 GridSubdivisionsPerLevel;
 	static const int32 GridLevels;
-
+	static const int32 RulerSubdivisionLevel;
 
 	// Background
 	static const float BackgroundQuadDepthOffset;
 
+	//---------------------------
+	//  Common Utility Methods
+	//---------------------------
+
+
+	// When creating UV unwraps or transforming between UV coordinate and positions in the 2D viewport,
+	// these functions will determine the mapping between UV values and the
+	// resulting unwrap mesh vertex positions. 
+	// If we're looking down on the unwrapped mesh, with the Z axis towards us, we want U's to be right, and
+	// V's to be up. In Unreal's left-handed coordinate system, this means that we map U's to world Y
+	// and V's to world X.
+	// Also, Unreal changes the V coordinates of imported meshes to 1-V internally, and we undo this
+	// while displaying the UV's because the users likely expect to see the original UV's (it would
+	// be particularly confusing for users working with UDIM assets, where internally stored V's 
+	// frequently end up negative).
+	// The ScaleFactor just scales the mesh up. Scaling the mesh up makes it easier to zoom in
+	// further into the display before getting issues with the camera near plane distance.
+	static FVector2f ExternalUVToInternalUV(const FVector2f& UV);
+
+	static FVector2f InternalUVToExternalUV(const FVector2f& UV);
+
+	static FVector3d UVToVertPosition(const FVector2f& UV);
+
+	static FVector2f VertPositionToUV(const FVector3d& VertPosition);
+
+	//--------------------------------
+	// CVARs for Experimental Features
+	//--------------------------------
+
+	static TAutoConsoleVariable<int32> CVarEnablePrototypeUDIMSupport;
+
 private:
 
 	FUVEditorUXSettings();
diff --git a/Engine/Plugins/Experimental/UVEditor/Source/UVEditorTools/Public/UVToolContextObjects.h b/Engine/Plugins/Experimental/UVEditor/Source/UVEditorTools/Public/UVToolContextObjects.h
--- a/Engine/Plugins/Experimental/UVEditor/Source/UVEditorTools/Public/UVToolContextObjects.h
+++ b/Engine/Plugins/Experimental/UVEditor/Source/UVEditorTools/Public/UVToolContextObjects.h
@@ -115,6 +115,97 @@ protected:
 	TUniqueFunction<void(FViewCameraState& CameraStateOut)> GetLivePreviewCameraStateFunc;
 };
 
+USTRUCT()
+struct UVEDITORTOOLS_API FUDIMBlock
+{
+	GENERATED_BODY();
+
+	UPROPERTY()
+	int32 BlockX;
+
+	UPROPERTY()
+	int32 BlockY;
+
+	UPROPERTY()
+	int32 SizeX;
+
+	UPROPERTY()
+	int32 SizeY;
+};
+
+/**
+ * Allows tools to interact with the 2d preview viewport 
+ */
+UCLASS()
+class UVEDITORTOOLS_API UUVTool2DViewportAPI : public UUVToolContextObject
+{
+	GENERATED_BODY()
+public:
+
+	void SetUDIMBlocks(TArray<FUDIMBlock>& BlocksIn, bool bBroadcast = true)
+	{
+		UDIMBlocks = BlocksIn;
+		if (bBroadcast)
+		{
+			OnUDIMBlockChange.Broadcast(UDIMBlocks);
+		}
+	}
+
+	const TArray<FUDIMBlock>& GetUDIMBlocks() const
+	{
+		return UDIMBlocks;
+	}
+
+	void SetDrawGrid(bool bDrawGridIn, bool bBroadcast = true)
+	{
+		bDrawGrid = bDrawGridIn;
+		if (bBroadcast)
+		{
+			OnDrawGridChange.Broadcast(bDrawGrid);
+		}
+	}
+
+	const bool GetDrawGrid() const
+    {
+		return bDrawGrid;
+	}
+
+	void SetDrawRulers(bool bDrawRulersIn, bool bBroadcast = true)
+	{
+		bDrawRulers = bDrawRulersIn;
+		if (bBroadcast)
+		{
+			OnDrawRulersChange.Broadcast(bDrawRulers);
+		}
+	}
+
+	const bool GetDrawRulers() const
+    {
+		return bDrawRulers;
+	}
+
+
+	DECLARE_MULTICAST_DELEGATE_OneParam(FOnUDIMBlockChange, const TArray<FUDIMBlock>& UDIMBlocks);
+	FOnUDIMBlockChange OnUDIMBlockChange;
+
+	DECLARE_MULTICAST_DELEGATE_OneParam(FOnDrawGridChange, bool bDrawGrid);
+	FOnDrawGridChange OnDrawGridChange;
+
+	DECLARE_MULTICAST_DELEGATE_OneParam(FOnDrawRulersChange, bool bDrawRulers);
+	FOnDrawRulersChange OnDrawRulersChange;
+
+protected:
+
+	UPROPERTY(Transient)
+	TArray<FUDIMBlock> UDIMBlocks;
+
+	UPROPERTY(Transient)
+	bool bDrawGrid;
+
+	UPROPERTY(Transient)
+	bool bDrawRulers;
+};
+
 /**
  * Allows tools to interact with buttons in the viewport
  */
