UVEditor: Fix QAGame cooking errors due to new UVEditor support for UDIMs

diff --git a/Engine/Plugins/Experimental/UVEditor/Source/UVEditor/Public/UVEditorBackgroundPreview.h b/Engine/Plugins/Experimental/UVEditor/Source/UVEditor/Public/UVEditorBackgroundPreview.h
--- a/Engine/Plugins/Experimental/UVEditor/Source/UVEditor/Public/UVEditorBackgroundPreview.h
+++ b/Engine/Plugins/Experimental/UVEditor/Source/UVEditor/Public/UVEditorBackgroundPreview.h
@@ -27,16 +27,16 @@ struct UVEDITOR_API FUVEditorBackgroundUDIMBlockInfo
 	GENERATED_BODY()
 
 	UPROPERTY()
-	int32 BlockX;
+	int32 BlockX = 1;
 
 	UPROPERTY()
-	int32 BlockY;
+	int32 BlockY = 1;
 
 	UPROPERTY()
-	int32 SizeX;
+	int32 SizeX = 256;
 
 	UPROPERTY()
-	int32 SizeY;
+	int32 SizeY = 256;
 };
 
 USTRUCT()
diff --git a/Engine/Plugins/Experimental/UVEditor/Source/UVEditorTools/Public/UVToolContextObjects.h b/Engine/Plugins/Experimental/UVEditor/Source/UVEditorTools/Public/UVToolContextObjects.h
--- a/Engine/Plugins/Experimental/UVEditor/Source/UVEditorTools/Public/UVToolContextObjects.h
+++ b/Engine/Plugins/Experimental/UVEditor/Source/UVEditorTools/Public/UVToolContextObjects.h
@@ -121,16 +121,16 @@ struct UVEDITORTOOLS_API FUDIMBlock
 	GENERATED_BODY();
 
 	UPROPERTY()
-	int32 BlockX;
+	int32 BlockX = 1;
 
 	UPROPERTY()
-	int32 BlockY;
+	int32 BlockY = 1;
 
 	UPROPERTY()
-	int32 SizeX;
+	int32 SizeX = 256;
 
 	UPROPERTY()
-	int32 SizeY;
+	int32 SizeY = 256;
 };
 
 /**
