UVEditor: Fix non-unity build issues related to new UDIM support

diff --git a/Engine/Plugins/Experimental/UVEditor/Source/UVEditorTools/Public/UVEditorUXSettings.h b/Engine/Plugins/Experimental/UVEditor/Source/UVEditorTools/Public/UVEditorUXSettings.h
--- a/Engine/Plugins/Experimental/UVEditor/Source/UVEditorTools/Public/UVEditorUXSettings.h
+++ b/Engine/Plugins/Experimental/UVEditor/Source/UVEditorTools/Public/UVEditorUXSettings.h
@@ -4,7 +4,7 @@
 #include "CoreMinimal.h"
 
 #include "CoreGlobals.h"
-
+#include "HAL/IConsoleManager.h"
 class UVEDITORTOOLS_API FUVEditorUXSettings
 {
 public:
