UVEditor: Add support for configuring visual UDIMs in the UVEditor 2D viewport based on a provided virtual texture or from a preconfigured polygroup on a mesh asset.

diff --git a/Engine/Plugins/Experimental/UVEditor/Source/UVEditor/Private/UVEditor2DViewportClient.cpp b/Engine/Plugins/Experimental/UVEditor/Source/UVEditor/Private/UVEditor2DViewportClient.cpp
--- a/Engine/Plugins/Experimental/UVEditor/Source/UVEditor/Private/UVEditor2DViewportClient.cpp
+++ b/Engine/Plugins/Experimental/UVEditor/Source/UVEditor/Private/UVEditor2DViewportClient.cpp
@@ -16,7 +16,6 @@
 #include "Engine/Canvas.h"
 #include "Math/Box.h"
 #include "MathUtil.h"
-#include "UDIMUtilities.h"
 
 namespace FUVEditor2DViewportClientLocals {
 
@@ -488,14 +487,14 @@ void FUVEditor2DViewportClient::DrawUDIMLabels(FViewport& InViewport, FSceneView
 	for (const FUDIMBlock& Block : UVTool2DViewportAPI->GetUDIMBlocks())
 	{
 		FVector2D TextSize;
-		FText UDIMLabel = FText::AsNumber(UE::TextureUtilitiesCommon::GetUDIMIndex(Block.BlockX, Block.BlockY), &FormatOptions);
+		FText UDIMLabel = FText::AsNumber(Block.UDIM, &FormatOptions);
 		Canvas.TextSize(Font, UDIMLabel.ToString(), TextSize[0], TextSize[1]);
 		if(TextSize[0] > MaxAllowedLabelWidth)
 		{
 			continue; // Don't draw if our label is bigger than the visual size of one UDIM "block" in the viewport
 		}			
 
-		FVector2D UDimUV(Block.BlockX + 1.0, Block.BlockY + 1.0);
+		FVector2D UDimUV(Block.BlockU() + 1.0, Block.BlockV() + 1.0);
 		FVector2D UDimPixel;
 		FUVEditor2DViewportClientLocals::ConvertUVToPixel(UDimUV, UDimPixel, View);
 
diff --git a/Engine/Plugins/Experimental/UVEditor/Source/UVEditor/Private/UVEditorBackgroundPreview.cpp b/Engine/Plugins/Experimental/UVEditor/Source/UVEditor/Private/UVEditorBackgroundPreview.cpp
--- a/Engine/Plugins/Experimental/UVEditor/Source/UVEditor/Private/UVEditorBackgroundPreview.cpp
+++ b/Engine/Plugins/Experimental/UVEditor/Source/UVEditor/Private/UVEditorBackgroundPreview.cpp
@@ -8,6 +8,7 @@
 #include "ToolSetupUtil.h"
 #include "Async/Async.h"
 #include "UVEditorUXSettings.h"
+#include "UDIMUtilities.h"
 
 using namespace UE::Geometry;
 
@@ -18,6 +19,7 @@ void UUVEditorBackgroundPreview::OnCreated()
 	Settings->WatchProperty(Settings->SourceType, [this](EUVEditorBackgroundSourceType) { bSettingsModified = true; });
 	Settings->WatchProperty(Settings->SourceTexture, [this](UTexture2D*) {bSettingsModified = true; });
 	Settings->WatchProperty(Settings->SourceMaterial, [this](UMaterial*) {bSettingsModified = true; });
+	Settings->WatchProperty(Settings->UDIMBlocks, [this](const TArray<int32>&) {bSettingsModified = true; });
 
 	bSettingsModified = false;
 
@@ -28,6 +30,14 @@ void UUVEditorBackgroundPreview::OnCreated()
 
 void UUVEditorBackgroundPreview::OnTick(float DeltaTime)
 {
+	// Check if the CVAR has been updated behind the scenes
+	bool bEnableUDIMSupport = (FUVEditorUXSettings::CVarEnablePrototypeUDIMSupport.GetValueOnGameThread() > 0);
+	if (Settings->bUDIMsEnabled != bEnableUDIMSupport)
+	{
+		Settings->bUDIMsEnabled = bEnableUDIMSupport;
+		bSettingsModified = true;
+	}
+
 	if (bSettingsModified)
 	{		
 		UpdateBackground();
@@ -42,7 +52,6 @@ void UUVEditorBackgroundPreview::UpdateVisibility()
 	if (Settings->bVisible == false)
 	{
 		BackgroundComponent->SetVisibility(false);
-		ActiveUDIMBlocks.UDIMBlocks.SetNum(0);
 		return;
 	}
 
@@ -52,15 +61,12 @@ void UUVEditorBackgroundPreview::UpdateVisibility()
 
 void UUVEditorBackgroundPreview::UpdateBackground()
 {
-	bool bEnableUDIMSupport = (FUVEditorUXSettings::CVarEnablePrototypeUDIMSupport.GetValueOnGameThread() > 0);
-
 	const FVector Normal(0, 0, 1);
 	const FColor BackgroundColor = FColor::Blue;
 
 	UMaterial* Material = LoadObject<UMaterial>(nullptr, TEXT("/UVEditor/Materials/UVEditorBackground"));
 	check(Material);	
 	BackgroundMaterial = UMaterialInstanceDynamic::Create(Material, this);
-	ActiveUDIMBlocks.UDIMBlocks.SetNum(0);
 	switch (Settings->SourceType)
 	{
 		case EUVEditorBackgroundSourceType::Checkerboard:
@@ -86,21 +92,6 @@ void UUVEditorBackgroundPreview::UpdateBackground()
 				{
 					BackgroundMaterial->SetTextureParameterValue(TEXT("BackgroundVTBaseMap"), Settings->SourceTexture);
 					BackgroundMaterial->SetScalarParameterValue(TEXT("BackgroundVirtualTextureSwitch"), 1);
-
-					// Check for UDIMs
-					if (bEnableUDIMSupport && Settings->SourceTexture->Source.GetNumBlocks() > 1) {
-						ActiveUDIMBlocks.UDIMBlocks.SetNum(Settings->SourceTexture->Source.GetNumBlocks());						
-					
-						for (int32 Block = 0; Block < Settings->SourceTexture->Source.GetNumBlocks(); ++Block) {
-							FTextureSourceBlock SourceBlock;
-							Settings->SourceTexture->Source.GetBlock(Block, SourceBlock);
-
-							ActiveUDIMBlocks.UDIMBlocks[Block].BlockX = SourceBlock.BlockX;
-							ActiveUDIMBlocks.UDIMBlocks[Block].BlockY = SourceBlock.BlockY;
-							ActiveUDIMBlocks.UDIMBlocks[Block].SizeX = SourceBlock.SizeX;
-							ActiveUDIMBlocks.UDIMBlocks[Block].SizeY = SourceBlock.SizeY;
-						}
-					}
 				}
 				else
 				{
@@ -119,12 +110,18 @@ void UUVEditorBackgroundPreview::UpdateBackground()
 	BackgroundComponent->Clear();
 
 	TArray<FVector2f> UDimBlocksToRender;
-	for (int32 BlockIndex = 0; BlockIndex < ActiveUDIMBlocks.UDIMBlocks.Num(); ++BlockIndex)
+	if (Settings->bUDIMsEnabled)
 	{
-		UDimBlocksToRender.Push(FVector2f(
-			ActiveUDIMBlocks.UDIMBlocks[BlockIndex].BlockX,
-			ActiveUDIMBlocks.UDIMBlocks[BlockIndex].BlockY
-			));
+		// TODO: Find a way to access the list of UDIMs from the context object instead? 
+		for (int32 BlockIndex = 0; BlockIndex < Settings->UDIMBlocks.Num(); ++BlockIndex)
+		{
+			FVector2f Block;
+			int32 UCoord, VCoord;
+			UE::TextureUtilitiesCommon::ExtractUDIMCoordinates(Settings->UDIMBlocks[BlockIndex], UCoord, VCoord);
+			Block.X = UCoord;
+			Block.Y = VCoord;
+			UDimBlocksToRender.Push(Block);
+		}
 	}
 	if (UDimBlocksToRender.Num() == 0)
 	{
diff --git a/Engine/Plugins/Experimental/UVEditor/Source/UVEditor/Private/UVEditorMode.cpp b/Engine/Plugins/Experimental/UVEditor/Source/UVEditor/Private/UVEditorMode.cpp
--- a/Engine/Plugins/Experimental/UVEditor/Source/UVEditor/Private/UVEditorMode.cpp
+++ b/Engine/Plugins/Experimental/UVEditor/Source/UVEditor/Private/UVEditorMode.cpp
@@ -38,6 +38,8 @@
 #include "UVEditorToolAnalyticsUtils.h"
 #include "Editor.h"
 #include "UVEditorUXSettings.h"
+#include "UDIMUtilities.h"
+#include "EditorSupportDelegates.h"
 
 #define LOCTEXT_NAMESPACE "UUVEditorMode"
 
@@ -134,6 +136,69 @@ const FToolTargetTypeRequirements& UUVEditorMode::GetToolTargetRequirements()
 	return ToolTargetRequirements;
 }
 
+void UUVEditorUDIMProperties::PostAction(EUVEditorModeActions Action)
+{
+	if (Action == EUVEditorModeActions::ConfigureUDIMsFromTexture)
+	{
+		UpdateActiveUDIMsFromTexture();
+	}
+	if (Action == EUVEditorModeActions::ConfigureUDIMsFromAsset)
+	{
+		UpdateActiveUDIMsFromAsset();
+	}
+}
+
+const TArray<FString>& UUVEditorUDIMProperties::GetAssetNames()
+{
+	return UVAssetNames;
+}
+
+int32 UUVEditorUDIMProperties::AssetByIndex() const
+{
+	return UVAssetNames.Find(UDIMSourceAsset);
+}
+
+void UUVEditorUDIMProperties::InitializeAssets(const TArray<TObjectPtr<UUVEditorToolMeshInput>>& TargetsIn)
+{
+	UVAssetNames.Reset(TargetsIn.Num());
+
+	for (int i = 0; i < TargetsIn.Num(); ++i)
+	{
+		UVAssetNames.Add(UE::ToolTarget::GetHumanReadableName(TargetsIn[i]->SourceTarget));
+	}
+}
+
+void UUVEditorUDIMProperties::UpdateActiveUDIMsFromTexture()
+{
+	ActiveUDIMs.Empty();
+	if (UDIMSourceTexture && UDIMSourceTexture->IsCurrentlyVirtualTextured() && UDIMSourceTexture->Source.GetNumBlocks() > 1)
+	{
+		for (int32 Block = 0; Block < UDIMSourceTexture->Source.GetNumBlocks(); ++Block)
+		{
+			FTextureSourceBlock SourceBlock;
+			UDIMSourceTexture->Source.GetBlock(Block, SourceBlock);
+
+			ActiveUDIMs.Add(FUDIMSpecifier({ UE::TextureUtilitiesCommon::GetUDIMIndex(SourceBlock.BlockX, SourceBlock.BlockY),
+				                             SourceBlock.BlockX, SourceBlock.BlockY }));
+		}
+	}
+	CheckAndUpdateWatched();
+
+	// Make sure we're updating the viewport immediately.
+	FEditorSupportDelegates::RedrawAllViewports.Broadcast();
+}
+
+void UUVEditorUDIMProperties::UpdateActiveUDIMsFromAsset()
+{
+	ActiveUDIMs.Empty();
+	int32 AssetID = AssetByIndex();
+	ParentMode->PopulateUDIMsByAsset(AssetID, ActiveUDIMs);
+	CheckAndUpdateWatched();
+
+	// Make sure we're updating the viewport immediately.
+	FEditorSupportDelegates::RedrawAllViewports.Broadcast();
+}
+
 UUVEditorMode::UUVEditorMode()
 {
 	Info = FEditorModeInfo(
@@ -157,13 +222,11 @@ void UUVEditorMode::Enter()
 	BackgroundVisualization->Settings->WatchProperty(BackgroundVisualization->Settings->bVisible, 
 		[this](bool IsVisible) {
 			UpdateTriangleMaterialBasedOnBackground(IsVisible);
-			UpdateViewportUDIMLabels(IsVisible);			
 		});
 
 	BackgroundVisualization->OnBackgroundMaterialChange.AddWeakLambda(this,
 		[this](TObjectPtr<UMaterialInstanceDynamic> MaterialInstance) {
 			UpdatePreviewMaterialBasedOnBackground();
-			UpdateViewportUDIMLabels(true);
 		});
 
 	PropertyObjectsToTick.Add(BackgroundVisualization->Settings);
@@ -192,6 +255,30 @@ void UUVEditorMode::Enter()
 
 	PropertyObjectsToTick.Add(UVEditorGridProperties);
 
+	UVEditorUDIMProperties = NewObject< UUVEditorUDIMProperties >(this);
+	UVEditorUDIMProperties->Initialize(this);
+	UDIMsChangedWatcherId = UVEditorUDIMProperties->WatchProperty(UVEditorUDIMProperties->ActiveUDIMs,
+		[this](const TArray<FUDIMSpecifier>& ActiveUDIMs) {
+			UpdateActiveUDIMs();
+		},
+		[](const TArray<FUDIMSpecifier>& ActiveUDIMsOld, const TArray<FUDIMSpecifier>& ActiveUDIMsNew) {
+			TSet<FUDIMSpecifier> NewCopy(ActiveUDIMsNew);
+			TSet<FUDIMSpecifier> OldCopy(ActiveUDIMsOld);
+
+			if (OldCopy.Num() != NewCopy.Num())
+			{
+				return true;
+			}
+			TSet<FUDIMSpecifier> Test = OldCopy.Union(NewCopy);
+			if (Test.Num() != OldCopy.Num())
+			{
+				return true;
+			}
+			return false;
+		});
+
+	PropertyObjectsToTick.Add(UVEditorUDIMProperties);
+
 	RegisterTools();
 
 	ActivateDefaultTool();
@@ -277,6 +364,14 @@ UObject* UUVEditorMode::GetGridSettingsObject()
 	return nullptr;
 }
 
+UObject* UUVEditorMode::GetUDIMSettingsObject()
+{
+	if (UVEditorUDIMProperties)
+	{
+		return UVEditorUDIMProperties;
+	}
+	return nullptr;
+}
 
 void UUVEditorMode::ActivateDefaultTool()
 {
@@ -578,6 +673,10 @@ void UUVEditorMode::InitializeTargets(const TArray<TObjectPtr<UObject>>& AssetsI
 	// Prep things for layer/channel selection
 	InitializeAssetNames(ToolTargets, AssetNames);
 	PendingUVLayerIndex.SetNumZeroed(ToolTargets.Num());
+	if (UVEditorUDIMProperties)
+	{
+		UVEditorUDIMProperties->InitializeAssets(ToolInputObjects);
+	}
 }
 
 void UUVEditorMode::EmitToolIndependentObjectChange(UObject* TargetObject, TUniquePtr<FToolCommandChange> Change, const FText& Description)
@@ -616,7 +715,8 @@ void UUVEditorMode::ApplyChanges()
 	GetToolManager()->EndUndoTransaction();
 }
 
-void UUVEditorMode::UpdateViewportUDIMLabels(bool IsVisible)
+
+void UUVEditorMode::UpdateActiveUDIMs()
 {
 	bool bEnableUDIMSupport = (FUVEditorUXSettings::CVarEnablePrototypeUDIMSupport.GetValueOnGameThread() > 0);
 	if (!bEnableUDIMSupport)
@@ -624,26 +724,80 @@ void UUVEditorMode::UpdateViewportUDIMLabels(bool IsVisible)
 		return;
 	}
 
+	if (!UVEditorUDIMProperties)
+	{
+		return;
+	}
+
+	TSet<FUDIMSpecifier> ActiveUDIMs(UVEditorUDIMProperties->ActiveUDIMs);
 	UContextObjectStore* ContextStore = GetInteractiveToolsContext()->ToolManager->GetContextObjectStore();
 	UUVTool2DViewportAPI* UVTool2DViewportAPI = ContextStore->FindContext<UUVTool2DViewportAPI>();
 	if (UVTool2DViewportAPI)
 	{
 		TArray<FUDIMBlock> Blocks;
-		if (IsVisible)
+		if (ActiveUDIMs.Num() > 0)
 		{
-			Blocks.SetNum(BackgroundVisualization->ActiveUDIMBlocks.UDIMBlocks.Num());
-			for (int32 BlockIndex = 0; BlockIndex < BackgroundVisualization->ActiveUDIMBlocks.UDIMBlocks.Num(); ++BlockIndex)
+			Blocks.Reserve(ActiveUDIMs.Num());
+			for (FUDIMSpecifier& UDIMSpecifier : ActiveUDIMs)
 			{
-				Blocks[BlockIndex].BlockX = BackgroundVisualization->ActiveUDIMBlocks.UDIMBlocks[BlockIndex].BlockX;
-				Blocks[BlockIndex].BlockY = BackgroundVisualization->ActiveUDIMBlocks.UDIMBlocks[BlockIndex].BlockY;
-				Blocks[BlockIndex].SizeX = BackgroundVisualization->ActiveUDIMBlocks.UDIMBlocks[BlockIndex].SizeX;
-				Blocks[BlockIndex].SizeY = BackgroundVisualization->ActiveUDIMBlocks.UDIMBlocks[BlockIndex].SizeY;
+				Blocks.Add({ UDIMSpecifier.UDIM });
+				UDIMSpecifier.UCoord = Blocks.Last().BlockU();
+				UDIMSpecifier.VCoord = Blocks.Last().BlockV();
 			}
 		}
 		UVTool2DViewportAPI->SetUDIMBlocks(Blocks, true);
 	}
+	UVEditorUDIMProperties->ActiveUDIMs = ActiveUDIMs.Array();
+	UVEditorUDIMProperties->SilentUpdateWatcherAtIndex(UDIMsChangedWatcherId);
+
+	if (BackgroundVisualization)
+	{
+		BackgroundVisualization->Settings->UDIMBlocks.Empty();
+		BackgroundVisualization->Settings->UDIMBlocks.Reserve(ActiveUDIMs.Num());
+		for (FUDIMSpecifier& UDIMSpecifier : ActiveUDIMs)
+		{
+			BackgroundVisualization->Settings->UDIMBlocks.Add(UDIMSpecifier.UDIM);
+		}
+		BackgroundVisualization->Settings->CheckAndUpdateWatched();
+	}
 }
 
+void UUVEditorMode::PopulateUDIMsByAsset(int32 AssetId, TArray<FUDIMSpecifier>& UDIMsOut) const
+{
+	UDIMsOut.Empty();
+	if (AssetId < 0 || AssetId >= ToolInputObjects.Num())
+	{
+		return;
+	}
+	if (ToolInputObjects[AssetId]->AppliedCanonical->HasAttributes())
+	{
+		for (int32 GroupLayerIndex = 0; GroupLayerIndex < ToolInputObjects[AssetId]->AppliedCanonical->Attributes()->NumPolygroupLayers(); ++GroupLayerIndex)
+		{
+			if (ToolInputObjects[AssetId]->AppliedCanonical->Attributes()->GetPolygroupLayer(GroupLayerIndex)->GetName() == "UDIM")
+			{
+				FDynamicMeshPolygroupAttribute* UDIMPolygroup = ToolInputObjects[AssetId]->AppliedCanonical->Attributes()->GetPolygroupLayer(GroupLayerIndex);
+
+				// This isn't very efficient, as a linear pass over all Tids. But there doesn't seem to be a better way to list all the groups present.
+				for (int32 Tid : ToolInputObjects[AssetId]->AppliedCanonical->TriangleIndicesItr())
+				{
+					int32 UDIMValue = UDIMPolygroup->GetValue(Tid);
+
+					// Discard any invalid UDIM group labels
+					if (UDIMValue < 1001)
+					{
+						continue;
+					}
+					FUDIMSpecifier UDIMSpecifier;
+					UDIMSpecifier.UDIM = UDIMValue;
+					UE::TextureUtilitiesCommon::ExtractUDIMCoordinates(UDIMValue, UDIMSpecifier.UCoord, UDIMSpecifier.VCoord);
+					UDIMsOut.AddUnique(UDIMSpecifier);
+				}
+			}
+
+			break; // Skip looking at other polygroup layers.
+		}
+	}
+}
 
 void UUVEditorMode::UpdateTriangleMaterialBasedOnBackground(bool IsBackgroundVisible)
 {
@@ -756,6 +910,7 @@ void UUVEditorMode::ModeTick(float DeltaTime)
 		ToolInput->AppliedPreview->Tick(DeltaTime);
 		ToolInput->UnwrapPreview->Tick(DeltaTime);
 	}
+
 }
 
 int32 UUVEditorMode::GetNumUVChannels(int32 AssetID) const
diff --git a/Engine/Plugins/Experimental/UVEditor/Source/UVEditor/Private/UVEditorModeToolkit.cpp b/Engine/Plugins/Experimental/UVEditor/Source/UVEditor/Private/UVEditorModeToolkit.cpp
--- a/Engine/Plugins/Experimental/UVEditor/Source/UVEditor/Private/UVEditorModeToolkit.cpp
+++ b/Engine/Plugins/Experimental/UVEditor/Source/UVEditor/Private/UVEditorModeToolkit.cpp
@@ -287,40 +287,23 @@ TSharedRef<SWidget> FUVEditorModeToolkit::CreateChannelMenu()
 
 TSharedRef<SWidget> FUVEditorModeToolkit::CreateBackgroundSettingsWidget()
 {
-	TSharedRef<SBorder> BackgroundDetailsContainer =
-		SNew(SBorder)
-		.BorderImage(FAppStyle::Get().GetBrush("NoBorder"));
-
-	TSharedRef<SWidget> Widget = SNew(SBorder)
-		.HAlign(HAlign_Fill)
-		.Padding(4)
-		[
-			SNew(SBox)			
-				.MinDesiredWidth(500)				
-			[
-				BackgroundDetailsContainer				
-			]
-		];
-
-	FPropertyEditorModule& PropertyEditorModule = FModuleManager::GetModuleChecked<FPropertyEditorModule>("PropertyEditor");
-
-	FDetailsViewArgs BackgroundDetailsViewArgs;
-	BackgroundDetailsViewArgs.bAllowSearch = false;
-	BackgroundDetailsViewArgs.NameAreaSettings = FDetailsViewArgs::HideNameArea;
-	BackgroundDetailsViewArgs.bHideSelectionTip = true;
-	BackgroundDetailsViewArgs.DefaultsOnlyVisibility = EEditDefaultsOnlyNodeVisibility::Automatic;
-	BackgroundDetailsViewArgs.bShowOptions = false;
-	BackgroundDetailsViewArgs.bAllowMultipleTopLevelObjects = false;
+	UUVEditorMode* Mode = Cast<UUVEditorMode>(GetScriptableEditorMode());
+	return CreateDisplaySettingsWidget(Mode->GetBackgroundSettingsObject());
+}
 
+TSharedRef<SWidget> FUVEditorModeToolkit::CreateGridSettingsWidget()
+{
 	UUVEditorMode* Mode = Cast<UUVEditorMode>(GetScriptableEditorMode());
-	TSharedRef<IDetailsView> BackgroundDetailsView = PropertyEditorModule.CreateDetailView(BackgroundDetailsViewArgs);
-	BackgroundDetailsView->SetObject(Mode->GetBackgroundSettingsObject());
-	BackgroundDetailsContainer->SetContent(BackgroundDetailsView);
+	return CreateDisplaySettingsWidget(Mode->GetGridSettingsObject());
+}
 
-	return Widget;
+TSharedRef<SWidget> FUVEditorModeToolkit::CreateUDIMSettingsWidget()
+{
+	UUVEditorMode* Mode = Cast<UUVEditorMode>(GetScriptableEditorMode());
+	return CreateDisplaySettingsWidget(Mode->GetUDIMSettingsObject());
 }
 
-TSharedRef<SWidget> FUVEditorModeToolkit::CreateGridSettingsWidget()
+TSharedRef<SWidget> FUVEditorModeToolkit::CreateDisplaySettingsWidget(UObject* SettingsObject) const
 {
 	TSharedRef<SBorder> GridDetailsContainer =
 		SNew(SBorder)
@@ -347,14 +330,15 @@ TSharedRef<SWidget> FUVEditorModeToolkit::CreateGridSettingsWidget()
 	GridDetailsViewArgs.bShowOptions = false;
 	GridDetailsViewArgs.bAllowMultipleTopLevelObjects = false;
 
-	UUVEditorMode* Mode = Cast<UUVEditorMode>(GetScriptableEditorMode());
 	TSharedRef<IDetailsView> GridDetailsView = PropertyEditorModule.CreateDetailView(GridDetailsViewArgs);
-	GridDetailsView->SetObject(Mode->GetGridSettingsObject());
+	GridDetailsView->SetObject(SettingsObject);
 	GridDetailsContainer->SetContent(GridDetailsView);
 
 	return Widget;
+
 }
 
+
 void FUVEditorModeToolkit::UpdateActiveToolProperties()
 {
 	UInteractiveTool* CurTool = GetScriptableEditorMode()->GetToolManager()->GetActiveTool(EToolSide::Left);
diff --git a/Engine/Plugins/Experimental/UVEditor/Source/UVEditor/Private/UVEditorToolkit.cpp b/Engine/Plugins/Experimental/UVEditor/Source/UVEditor/Private/UVEditorToolkit.cpp
--- a/Engine/Plugins/Experimental/UVEditor/Source/UVEditor/Private/UVEditorToolkit.cpp
+++ b/Engine/Plugins/Experimental/UVEditor/Source/UVEditor/Private/UVEditorToolkit.cpp
@@ -464,6 +464,21 @@ void FUVEditorToolkit::PostInitAssetEditor()
 				]
 				];
 
+			bool bEnableUDIMSupport = (FUVEditorUXSettings::CVarEnablePrototypeUDIMSupport.GetValueOnGameThread() > 0);
+			if (bEnableUDIMSupport)
+			{
+				Container->AddSlot()
+					.AutoHeight()
+					.Padding(FMargin(0.f, 0.f, 8.f, 0.f))
+					[
+						SNew(SBox)
+						.MinDesiredWidth(500)
+					[
+						UVModeToolkit->CreateUDIMSettingsWidget()
+					]
+					];
+			}
+
 			Container->AddSlot()
 				.AutoHeight()
 				.Padding(FMargin(0.f, 0.f, 8.f, 0.f))
diff --git a/Engine/Plugins/Experimental/UVEditor/Source/UVEditor/Public/UVEditorBackgroundPreview.h b/Engine/Plugins/Experimental/UVEditor/Source/UVEditor/Public/UVEditorBackgroundPreview.h
--- a/Engine/Plugins/Experimental/UVEditor/Source/UVEditor/Public/UVEditorBackgroundPreview.h
+++ b/Engine/Plugins/Experimental/UVEditor/Source/UVEditor/Public/UVEditorBackgroundPreview.h
@@ -21,33 +21,6 @@ enum class EUVEditorBackgroundSourceType
 	Material
 };
 
-USTRUCT()
-struct UVEDITOR_API FUVEditorBackgroundUDIMBlockInfo
-{
-	GENERATED_BODY()
-
-	UPROPERTY()
-	int32 BlockX = 1;
-
-	UPROPERTY()
-	int32 BlockY = 1;
-
-	UPROPERTY()
-	int32 SizeX = 256;
-
-	UPROPERTY()
-	int32 SizeY = 256;
-};
-
-USTRUCT()
-struct UVEDITOR_API FUVEditorBackgroundUDIMBlocks
-{
-	GENERATED_BODY()
-
-	UPROPERTY()
-	TArray<FUVEditorBackgroundUDIMBlockInfo> UDIMBlocks;
-};
-
 /**
  * Visualization settings for the UUVEditorBackgroundPreview
  */
@@ -71,6 +44,12 @@ public:
 	/** Display a background based on the selected material */
 	UPROPERTY(EditAnywhere, Category = Background, meta = (EditCondition = "SourceType==EUVEditorBackgroundSourceType::Material && bVisible", EditConditionHides = true))
 	TObjectPtr<UMaterial> SourceMaterial;
+
+	UPROPERTY()
+	TArray<int32> UDIMBlocks;
+
+	UPROPERTY()
+	bool bUDIMsEnabled = false;
 };
 
 /**
@@ -88,16 +67,12 @@ public:
 	 * Client must call this every frame for changes to .Settings to be reflected in rendered result.
 	 */
 	void OnTick(float DeltaTime);
-
-
+	
 public:
 	/** Visualization settings */
 	UPROPERTY()
 	TObjectPtr<UUVEditorBackgroundPreviewProperties> Settings;
 
-	UPROPERTY()
-	FUVEditorBackgroundUDIMBlocks ActiveUDIMBlocks;
-
 	/** The component containing the quad visualization */
 	UPROPERTY()
 	TObjectPtr<UTriangleSetComponent> BackgroundComponent;
diff --git a/Engine/Plugins/Experimental/UVEditor/Source/UVEditor/Public/UVEditorMode.h b/Engine/Plugins/Experimental/UVEditor/Source/UVEditor/Public/UVEditorMode.h
--- a/Engine/Plugins/Experimental/UVEditor/Source/UVEditor/Public/UVEditorMode.h
+++ b/Engine/Plugins/Experimental/UVEditor/Source/UVEditor/Public/UVEditorMode.h
@@ -29,7 +29,7 @@ class UUVEditorToolMeshInput;
 class UUVEditorBackgroundPreview;
 class UUVToolViewportButtonsAPI;
 class UUVTool2DViewportAPI;
-
+class UUVEditorMode;
 
 /**
  * Visualization settings for the UUVEditorMode's Grid
@@ -48,6 +48,88 @@ public:
 	bool bDrawRulers = true;
 };
 
+USTRUCT()
+struct FUDIMSpecifier
+{
+	GENERATED_USTRUCT_BODY()
+
+	UPROPERTY(EditAnywhere, Transient, Category = UDIM, meta = (ClampMin = "1001", UIMin = "1001"))
+	int32 UDIM = 1001;
+
+	UPROPERTY(VisibleAnywhere, Transient, Category = UDIM, meta = (DisplayName = "Block U Offset"))
+	int32 UCoord = 0;
+
+	UPROPERTY(VisibleAnywhere, Transient, Category = UDIM, meta = (DisplayName = "Block V Offset"))
+	int32 VCoord = 0;
+
+	friend bool operator==(const FUDIMSpecifier& A, const FUDIMSpecifier& B )
+	{
+		return A.UDIM == B.UDIM;
+	}
+
+	friend uint32 GetTypeHash(const FUDIMSpecifier& UDIMSpecifier)
+	{
+		uint32 HashCode = GetTypeHash(UDIMSpecifier.UDIM);
+		return HashCode;
+	}
+};
+
+UENUM()
+enum class EUVEditorModeActions
+{
+	NoAction,
+
+	ConfigureUDIMsFromAsset,
+	ConfigureUDIMsFromTexture
+};
+
+/**
+ * Settings for UDIMs in the UVEditor
+ */
+UCLASS()
+class UVEDITOR_API UUVEditorUDIMProperties : public UInteractiveToolPropertySet
+{
+	GENERATED_BODY()
+public:
+	void Initialize(UUVEditorMode* ParentModeIn) { ParentMode = ParentModeIn; }
+	void PostAction(EUVEditorModeActions Action);
+
+	UPROPERTY(EditAnywhere, Transient, Category = "UDIM", meta = (DisplayName = "UDIM Source Mesh", GetOptions = GetAssetNames))
+	FString UDIMSourceAsset;
+
+	UFUNCTION()
+	const TArray<FString>& GetAssetNames();
+
+	UFUNCTION()
+	int32 AssetByIndex() const;
+
+	/** Set UDIM Layout from selected asset's UVs */
+	UFUNCTION(CallInEditor, Category = UDIM)
+	void SetUDIMsFromAsset() { PostAction(EUVEditorModeActions::ConfigureUDIMsFromAsset); }
+
+	/** Texture asset to source UDIM information from */
+	UPROPERTY(EditAnywhere, Transient, Category = UDIM)
+	TObjectPtr<UTexture2D> UDIMSourceTexture;
+
+	/** Set UDIM Layout from selected texture asset */
+	UFUNCTION(CallInEditor, Category = UDIM)
+	void SetUDIMsFromTexture() { PostAction(EUVEditorModeActions::ConfigureUDIMsFromTexture); };
+
+	/** Currently active UDIM set */
+	UPROPERTY(EditAnywhere, Transient, Category = UDIM, meta = (TitleProperty = "UDIM"))
+	TArray<FUDIMSpecifier> ActiveUDIMs;
+
+public:
+	void InitializeAssets(const TArray<TObjectPtr<UUVEditorToolMeshInput>>& TargetsIn);
+
+private:
+	void UpdateActiveUDIMsFromTexture();
+	void UpdateActiveUDIMsFromAsset();
+
+	TWeakObjectPtr<UUVEditorMode> ParentMode;
+	TArray<FString> UVAssetNames;
+};
+
 /**
  * The UV editor mode is the mode used in the UV asset editor. It holds most of the inter-tool state.
  * We put things in a mode instead of directly into the asset editor in case we want to someday use the mode
@@ -123,6 +205,9 @@ public:
 	/** @return A settings object suitable for display in a details panel to control the grid. */
 	UObject* GetGridSettingsObject();
 
+	/** @return A settings object suitable for display in a details panel to control UDIM configuration. */
+	UObject* GetUDIMSettingsObject();
+
 	// UEdMode overrides
 	virtual void Enter() override;
 	virtual bool ShouldToolStartBeAllowed(const FString& ToolIdentifier) const override;
@@ -147,6 +232,13 @@ public:
 	// Hold a settings object to configure the grid
 	UPROPERTY()
 	TObjectPtr<UUVEditorGridProperties> UVEditorGridProperties = nullptr;
+
+	// Hold a settings object to configure the UDIMs
+	UPROPERTY()
+	TObjectPtr<UUVEditorUDIMProperties> UVEditorUDIMProperties = nullptr;
+
+	void PopulateUDIMsByAsset(int32 AssetId, TArray<FUDIMSpecifier>& UDIMsOut) const;
+
 protected:
 
 	// UEdMode overrides
@@ -158,7 +250,9 @@ protected:
 	
 	void UpdateTriangleMaterialBasedOnBackground(bool IsBackgroundVisible);
 	void UpdatePreviewMaterialBasedOnBackground();
-	void UpdateViewportUDIMLabels(bool IsVisible);
+
+	void UpdateActiveUDIMs();
+	int32 UDIMsChangedWatcherId;
 
 	/**
 	 * Stores original input objects, for instance UStaticMesh pointers. AssetIDs on tool input 
diff --git a/Engine/Plugins/Experimental/UVEditor/Source/UVEditor/Public/UVEditorModeToolkit.h b/Engine/Plugins/Experimental/UVEditor/Source/UVEditor/Public/UVEditorModeToolkit.h
--- a/Engine/Plugins/Experimental/UVEditor/Source/UVEditor/Public/UVEditorModeToolkit.h
+++ b/Engine/Plugins/Experimental/UVEditor/Source/UVEditor/Public/UVEditorModeToolkit.h
@@ -31,6 +31,9 @@ public:
 
 	/** Creates a widget where the grid display can be changed. */
 	virtual TSharedRef<SWidget> CreateGridSettingsWidget();
+	
+	/** Creates a widget where the visible UDIMs can be changed. */
+	virtual TSharedRef<SWidget> CreateUDIMSettingsWidget();
 
 	// FModeToolkit
 	virtual void Init(const TSharedPtr<IToolkitHost>& InitToolkitHost, TWeakObjectPtr<UEdMode> InOwningMode) override;
@@ -61,6 +64,8 @@ public:
 protected:
 	const static TArray<FName> PaletteNames_Standard;
 
+	TSharedRef<SWidget> CreateDisplaySettingsWidget(UObject* SettingsObject) const;
+
 	// FModeToolkit
 	virtual void OnToolStarted(UInteractiveToolManager* Manager, UInteractiveTool* Tool) override;
 	virtual void OnToolEnded(UInteractiveToolManager* Manager, UInteractiveTool* Tool) override;
diff --git a/Engine/Plugins/Experimental/UVEditor/Source/UVEditorTools/Private/UVToolContextObjects.cpp b/Engine/Plugins/Experimental/UVEditor/Source/UVEditorTools/Private/UVToolContextObjects.cpp
--- a/Engine/Plugins/Experimental/UVEditor/Source/UVEditorTools/Private/UVToolContextObjects.cpp
+++ b/Engine/Plugins/Experimental/UVEditor/Source/UVEditorTools/Private/UVToolContextObjects.cpp
@@ -7,6 +7,7 @@
 #include "InputRouter.h" // Need to define this and UWorld so weak pointers know they are UThings
 #include "InteractiveToolManager.h"
 #include "ToolTargets/UVEditorToolMeshInput.h"
+#include "UDIMUtilities.h"
 
 using namespace UE::Geometry;
 
@@ -83,6 +84,25 @@ void UUVToolLivePreviewAPI::Initialize(UWorld* WorldIn, UInputRouter* RouterIn,
 	GetLivePreviewCameraStateFunc = MoveTemp(GetLivePreviewCameraStateFuncIn);
 }
 
+int32 FUDIMBlock::BlockU() const
+{
+	int32 BlockU, BlockV;
+	UE::TextureUtilitiesCommon::ExtractUDIMCoordinates(UDIM, BlockU, BlockV);
+	return BlockU;
+}
+
+int32 FUDIMBlock::BlockV() const
+{
+	int32 BlockU, BlockV;
+	UE::TextureUtilitiesCommon::ExtractUDIMCoordinates(UDIM, BlockU, BlockV);
+	return BlockV;
+}
+
+void FUDIMBlock::SetFromBlocks(int32 BlockU, int32 BlockV)
+{
+	UDIM = UE::TextureUtilitiesCommon::GetUDIMIndex(BlockU, BlockV);
+}
+
 void UUVToolAABBTreeStorage::Set(FDynamicMesh3* MeshKey, TSharedPtr<FDynamicMeshAABBTree3> Tree, 
 	UUVEditorToolMeshInput* InputObject)
 {
diff --git a/Engine/Plugins/Experimental/UVEditor/Source/UVEditorTools/Public/UVToolContextObjects.h b/Engine/Plugins/Experimental/UVEditor/Source/UVEditorTools/Public/UVToolContextObjects.h
--- a/Engine/Plugins/Experimental/UVEditor/Source/UVEditorTools/Public/UVToolContextObjects.h
+++ b/Engine/Plugins/Experimental/UVEditor/Source/UVEditorTools/Public/UVToolContextObjects.h
@@ -121,16 +121,11 @@ struct UVEDITORTOOLS_API FUDIMBlock
 	GENERATED_BODY();
 
 	UPROPERTY()
-	int32 BlockX = 1;
+	int32 UDIM = 1001;
 
-	UPROPERTY()
-	int32 BlockY = 1;
-
-	UPROPERTY()
-	int32 SizeX = 256;
-
-	UPROPERTY()
-	int32 SizeY = 256;
+	int32 BlockU() const;
+	int32 BlockV() const;
+	void SetFromBlocks(int32 BlockU, int32 BlockV);
 };
 
 /**
diff --git a/Engine/Plugins/Experimental/UVEditor/Source/UVEditorTools/UVEditorTools.Build.cs b/Engine/Plugins/Experimental/UVEditor/Source/UVEditorTools/UVEditorTools.Build.cs
--- a/Engine/Plugins/Experimental/UVEditor/Source/UVEditorTools/UVEditorTools.Build.cs
+++ b/Engine/Plugins/Experimental/UVEditor/Source/UVEditorTools/UVEditorTools.Build.cs
@@ -49,6 +49,7 @@ public class UVEditorTools : ModuleRules
 				"GeometryCore",
 				"ModelingComponents",
 				"ModelingOperators",
+				"TextureUtilitiesCommon"
 				// ... add private dependencies that you statically link with here ...	
 			}
 			);
