UVEditor: Fix Non-unity build related to missing forward declaration in UVEditorMode.h

#rnx

diff --git a/Engine/Plugins/Experimental/UVEditor/Source/UVEditor/Public/UVEditorMode.h b/Engine/Plugins/Experimental/UVEditor/Source/UVEditor/Public/UVEditorMode.h
--- a/Engine/Plugins/Experimental/UVEditor/Source/UVEditor/Public/UVEditorMode.h
+++ b/Engine/Plugins/Experimental/UVEditor/Source/UVEditor/Public/UVEditorMode.h
@@ -30,6 +30,7 @@ class UUVEditorBackgroundPreview;
 class UUVToolViewportButtonsAPI;
 class UUVTool2DViewportAPI;
 class UUVEditorMode;
+class UTexture2D;
 
 /**
  * Visualization settings for the UUVEditorMode's Grid
