UVEditor: Fix crash when a hover in the live preview input router is not terminated before Shutdown in Seam tool.

diff --git a/Engine/Plugins/Experimental/UVEditor/Source/UVEditorTools/Private/UVEditorSeamTool.cpp b/Engine/Plugins/Experimental/UVEditor/Source/UVEditorTools/Private/UVEditorSeamTool.cpp
--- a/Engine/Plugins/Experimental/UVEditor/Source/UVEditorTools/Private/UVEditorSeamTool.cpp
+++ b/Engine/Plugins/Experimental/UVEditor/Source/UVEditorTools/Private/UVEditorSeamTool.cpp
@@ -563,6 +563,14 @@ void UUVEditorSeamTool::Shutdown(EToolShutdownType ShutdownType)
 {
 	TRACE_CPUPROFILER_EVENT_SCOPE(UVEditorSeamTool_Shutdown);
 
+	if (LivePreviewInputRouter.IsValid())
+	{
+		// TODO: Arguably the live preview input router should do this for us before Shutdown(), but
+		// we don't currently have support for that...
+		LivePreviewInputRouter->ForceTerminateSource(LivePreviewBehaviorSource);
+		LivePreviewInputRouter->DeregisterSource(LivePreviewBehaviorSource);
+	}
+
 	// Apply any pending seam if needed
 	if (ShutdownType != EToolShutdownType::Cancel && LockedPath.Num() > 0)
 	{
@@ -574,12 +582,7 @@ void UUVEditorSeamTool::Shutdown(EToolShutdownType ShutdownType)
 	EmitChangeAPI = nullptr;
 	LivePreviewAPI = nullptr;
 
-	if (LivePreviewInputRouter.IsValid())
-	{
-		LivePreviewInputRouter->DeregisterSource(LivePreviewBehaviorSource);
-	}
 	LivePreviewBehaviorSet->RemoveAll();
-
 	LivePreviewBehaviorSet = nullptr;
 	LivePreviewBehaviorSource = nullptr;
 
