UVEditor: Major UV selection refactoring.

- Move selection into a context object API.
- Factor out selection highlighting into a mechanic.
- Support multi-asset selections.
- Pull selection-based actions out of Select Tool, and set it as a default tool for gizmo transformation only.
- Change mode initialization to work via context objects.

diff --git a/Engine/Plugins/Experimental/UVEditor/Source/UVEditor/Private/DetailsCustomizations/UVSelectToolCustomizations.cpp b/Engine/Plugins/Experimental/UVEditor/Source/UVEditor/Private/DetailsCustomizations/UVSelectToolCustomizations.cpp
deleted file mode 100644
diff --git a/Engine/Plugins/Experimental/UVEditor/Source/UVEditor/Private/SUVEditor2DViewport.cpp b/Engine/Plugins/Experimental/UVEditor/Source/UVEditor/Private/SUVEditor2DViewport.cpp
--- a/Engine/Plugins/Experimental/UVEditor/Source/UVEditor/Private/SUVEditor2DViewport.cpp
+++ b/Engine/Plugins/Experimental/UVEditor/Source/UVEditor/Private/SUVEditor2DViewport.cpp
@@ -5,7 +5,7 @@
 #include "SUVEditor2DViewportToolBar.h"
 #include "UVEditor2DViewportClient.h"
 #include "UVEditorCommands.h"
-#include "UVToolContextObjects.h" // UUVToolViewportButtonsAPI::ESelectionMode
+#include "ContextObjects/UVToolContextObjects.h" // UUVToolViewportButtonsAPI::ESelectionMode
 
 #define LOCTEXT_NAMESPACE "SUVEditor2DViewport"
 
diff --git a/Engine/Plugins/Experimental/UVEditor/Source/UVEditor/Private/UVEditorCommands.cpp b/Engine/Plugins/Experimental/UVEditor/Source/UVEditor/Private/UVEditorCommands.cpp
--- a/Engine/Plugins/Experimental/UVEditor/Source/UVEditor/Private/UVEditorCommands.cpp
+++ b/Engine/Plugins/Experimental/UVEditor/Source/UVEditor/Private/UVEditorCommands.cpp
@@ -24,7 +24,6 @@ void FUVEditorCommands::RegisterCommands()
 	UI_COMMAND(ApplyChanges, "Apply", "Apply changes to original meshes", EUserInterfaceActionType::Button, FInputChord());
 
 	// These get linked to various tool buttons.
-	UI_COMMAND(BeginSelectTool, "Edit", "Selection-based UV editing", EUserInterfaceActionType::ToggleButton, FInputChord());
 	UI_COMMAND(BeginLayoutTool, "Layout", "Transform and pack existing UVs", EUserInterfaceActionType::ToggleButton, FInputChord());
 	UI_COMMAND(BeginParameterizeMeshTool, "AutoUV", "Auto-unwrap and pack UVs", EUserInterfaceActionType::ToggleButton, FInputChord());
 	UI_COMMAND(BeginChannelEditTool, "Channels", "Modify UV channels", EUserInterfaceActionType::ToggleButton, FInputChord());
diff --git a/Engine/Plugins/Experimental/UVEditor/Source/UVEditor/Private/UVEditorMode.cpp b/Engine/Plugins/Experimental/UVEditor/Source/UVEditor/Private/UVEditorMode.cpp
--- a/Engine/Plugins/Experimental/UVEditor/Source/UVEditor/Private/UVEditorMode.cpp
+++ b/Engine/Plugins/Experimental/UVEditor/Source/UVEditor/Private/UVEditorMode.cpp
@@ -3,8 +3,15 @@
 #include "UVEditorMode.h"
 
 #include "Algo/AnyOf.h"
+#include "Actions/UVIslandConformalUnwrapAction.h"
+#include "Actions/UVSplitAction.h"
+#include "Actions/UVSeamSewAction.h"
+#include "Actions/UVToolAction.h"
 #include "AssetEditorModeManager.h"
+#include "ContextObjects/UVToolViewportButtonsAPI.h"
+#include "ContextObjects/UVToolAssetInputsContext.h"
 #include "ContextObjectStore.h"
+#include "Selection/UVToolSelectionAPI.h"
 #include "Drawing/MeshElementsVisualizer.h"
 #include "Editor.h"
 #include "EditorViewportClient.h"
@@ -33,10 +40,9 @@
 #include "UVSelectTool.h"
 #include "UVEditorModeToolkit.h"
 #include "UVEditorSubsystem.h"
-#include "UVToolContextObjects.h"
+#include "ContextObjects/UVToolContextObjects.h"
 #include "UVEditorBackgroundPreview.h"
 #include "UVEditorToolAnalyticsUtils.h"
-#include "Editor.h"
 #include "UVEditorUXSettings.h"
 #include "UDIMUtilities.h"
 #include "EditorSupportDelegates.h"
@@ -72,6 +78,45 @@ namespace UVEditorModeLocals
 		CameraStateOut.Orientation = ViewRotation.Quaternion();
 	}
 
+	/**
+	 * Support for undoing a tool start in such a way that we go back to the mode's default
+	 * tool on undo.
+	 */
+	class FUVEditorBeginToolChange : public FToolCommandChange
+	{
+	public:
+		virtual void Apply(UObject* Object) override
+		{
+			// Do nothing, since we don't allow a re-do back into a tool
+		}
+
+		virtual void Revert(UObject* Object) override
+		{
+			UUVEditorMode* Mode = Cast<UUVEditorMode>(Object);
+			// Don't really need the check for default tool since we theoretically shouldn't
+			// be issuing this transaction for starting the default tool, but still...
+			if (Mode && !Mode->IsDefaultToolActive())
+			{
+				Mode->GetInteractiveToolsContext()->EndTool(EToolShutdownType::Cancel);
+				Mode->ActivateDefaultTool();
+			}
+		}
+
+		virtual bool HasExpired(UObject* Object) const override
+		{
+			// To not be expired, we must be in some tool.
+			UUVEditorMode* Mode = Cast<UUVEditorMode>(Object);
+			return !(Mode && Mode->GetInteractiveToolsContext()
+				&& Mode->GetInteractiveToolsContext()->ToolManager
+				&& Mode->GetInteractiveToolsContext()->ToolManager->HasAnyActiveTool());
+		}
+
+		virtual FString ToString() const override
+		{
+			return TEXT("UVEditorModeLocals::FUVEditorBeginToolChange");
+		}
+	};
+
 	/** 
 	 * Change for undoing/redoing displayed layer changes.
 	 */
@@ -216,6 +261,14 @@ void UUVEditorMode::Enter()
 {
 	Super::Enter();
 
+	// Our mode needs to get Render and DrawHUD calls, but doesn't implement the legacy interface that
+	// causes those to be called. Instead, we'll hook into the tools context's Render and DrawHUD calls.
+	GetInteractiveToolsContext()->OnRender.AddUObject(this, &UUVEditorMode::Render);
+	GetInteractiveToolsContext()->OnDrawHUD.AddUObject(this, &UUVEditorMode::DrawHUD);
+
+	InitializeModeContexts();
+	InitializeTargets();
+
 	BackgroundVisualization = NewObject<UUVEditorBackgroundPreview>(this);
 	BackgroundVisualization->CreateInWorld(GetWorld(), FTransform::Identity);
 
@@ -280,7 +333,7 @@ void UUVEditorMode::Enter()
 	PropertyObjectsToTick.Add(UVEditorUDIMProperties);
 
 	RegisterTools();
-
+	RegisterActions();
 	ActivateDefaultTool();
 
 	if (FEngineAnalytics::IsAvailable())
@@ -300,14 +353,16 @@ void UUVEditorMode::RegisterTools()
 {
 	const FUVEditorCommands& CommandInfos = FUVEditorCommands::Get();
 
-	// Note that the identifiers below need to match the command names so that the tool icons can 
-	// be easily retrieved from the active tool name in UVEditorModeToolkit::OnToolStarted. Otherwise
-	// we would need to keep some other mapping from tool identifier to tool icon.
-
+	// The Select Tool is a bit different because it runs between other tools. It doesn't have a button and is not
+	// hotkeyable, hence it doesn't have a command item. Instead we register it directly with the tool manager.
 	UUVSelectToolBuilder* UVSelectToolBuilder = NewObject<UUVSelectToolBuilder>();
 	UVSelectToolBuilder->Targets = &ToolInputObjects;
 	DefaultToolIdentifier = TEXT("BeginSelectTool");
-	RegisterTool(CommandInfos.BeginSelectTool, DefaultToolIdentifier, UVSelectToolBuilder);
+	GetToolManager()->RegisterToolType(DefaultToolIdentifier, UVSelectToolBuilder);
+
+	// Note that the identifiers below need to match the command names so that the tool icons can 
+	// be easily retrieved from the active tool name in UVEditorModeToolkit::OnToolStarted. Otherwise
+	// we would need to keep some other mapping from tool identifier to tool icon.
 
 	UUVEditorLayoutToolBuilder* UVEditorLayoutToolBuilder = NewObject<UUVEditorLayoutToolBuilder>();
 	UVEditorLayoutToolBuilder->Targets = &ToolInputObjects;
@@ -330,6 +385,44 @@ void UUVEditorMode::RegisterTools()
 	RegisterTool(CommandInfos.BeginRecomputeUVsTool, TEXT("BeginRecomputeUVsTool"), UVEditorRecomputeUVsToolBuilder);
 }
 
+void UUVEditorMode::RegisterActions()
+{
+	const FUVEditorCommands& CommandInfos = FUVEditorCommands::Get();
+	const TSharedRef<FUICommandList>& CommandList = Toolkit->GetToolkitCommands();
+
+	auto PrepAction = [this, CommandList](TSharedPtr<const FUICommandInfo> CommandInfo, UUVToolAction* Action)
+	{
+		Action->Setup(GetToolManager());
+		CommandList->MapAction(CommandInfo,
+			FExecuteAction::CreateWeakLambda(Action, [this, Action]() 
+			{ 
+				Action->ExecuteAction();
+				for (const FUVToolSelection& Selection : SelectionAPI->GetSelections())
+				{
+					// The actions are expected to have left the selection in a valid state. If this is not
+					// the case, we'll clear the selection here, which is still not a proper solution because
+					// undoing the transaction will still put things into an invalid state (but our Select Tool
+					// should be robust to that in its OnSelectionChanged).
+					if (!ensure(Selection.Target.IsValid() 
+						&& Selection.AreElementsPresentInMesh(*Selection.Target->UnwrapCanonical)))
+					{
+						SelectionAPI->ClearSelections(true, true);
+						SelectionAPI->ClearHighlight();
+						break;
+					}
+				}
+			}),
+			FCanExecuteAction::CreateWeakLambda(Action, [Action, this]() 
+				{ 
+					return IsDefaultToolActive() && Action->CanExecuteAction(); 
+				}));
+		RegisteredActions.Add(Action);
+	};
+	PrepAction(CommandInfos.SewAction, NewObject<UUVSeamSewAction>());
+	PrepAction(CommandInfos.SplitAction, NewObject<UUVSplitAction>());
+	PrepAction(CommandInfos.IslandConformalUnwrapAction, NewObject<UUVIslandConformalUnwrapAction>());
+}
+
 bool UUVEditorMode::ShouldToolStartBeAllowed(const FString& ToolIdentifier) const
 {
 	// For now we've decided to disallow switch-away on accept/cancel tools in the UV editor.
@@ -346,6 +439,43 @@ void UUVEditorMode::CreateToolkit()
 	Toolkit = MakeShared<FUVEditorModeToolkit>();
 }
 
+void UUVEditorMode::OnToolStarted(UInteractiveToolManager* Manager, UInteractiveTool* Tool)
+{
+	using namespace UVEditorModeLocals;
+
+	FText TransactionName = LOCTEXT("ActivateTool", "Activate Tool");
+	
+	if (!IsDefaultToolActive())
+	{
+		GetInteractiveToolsContext()->GetTransactionAPI()->BeginUndoTransaction(TransactionName);
+	
+		// If a tool doesn't support selection, we can't be certain that it won't put the meshes
+		// in a state where the selection refers to invalid elements.
+		if (!Cast<IUVToolSupportsSelection>(Tool))
+		{
+			SelectionAPI->ClearSelections(true, true); // broadcast, emit
+			SelectionAPI->ClearHighlight();
+		}
+	
+		GetInteractiveToolsContext()->GetTransactionAPI()->AppendChange(
+			this, MakeUnique<FUVEditorBeginToolChange>(), TransactionName);
+
+		GetInteractiveToolsContext()->GetTransactionAPI()->EndUndoTransaction();
+	}
+
+}
+
+void UUVEditorMode::OnToolEnded(UInteractiveToolManager* Manager, UInteractiveTool* Tool)
+{
+	for (TWeakObjectPtr<UUVToolContextObject> Context : ContextsToUpdateOnToolEnd)
+	{
+		if (ensure(Context.IsValid()))
+		{
+			Context->OnToolEnded(Tool);
+		}
+	}
+}
+
 UObject* UUVEditorMode::GetBackgroundSettingsObject()
 {
 	if (BackgroundVisualization)
@@ -373,6 +503,22 @@ UObject* UUVEditorMode::GetUDIMSettingsObject()
 	return nullptr;
 }
 
+void UUVEditorMode::Render(IToolsContextRenderAPI* RenderAPI)
+{
+	if (SelectionAPI)
+	{
+		SelectionAPI->Render(RenderAPI);
+	}
+}
+
+void UUVEditorMode::DrawHUD(FCanvas* Canvas, IToolsContextRenderAPI* RenderAPI)
+{
+	if (SelectionAPI)
+	{
+		SelectionAPI->DrawHUD(Canvas, RenderAPI);
+	}
+}
+
 void UUVEditorMode::ActivateDefaultTool()
 {
 	GetInteractiveToolsContext()->StartTool(DefaultToolIdentifier);
@@ -482,6 +628,12 @@ void UUVEditorMode::Exit()
 	// inaccessible, so we end up having to do this to force the shutdown right now.
 	GetToolManager()->DeactivateTool(EToolSide::Mouse, EToolShutdownType::Cancel);
 
+	for (UUVToolAction* Action : RegisteredActions)
+	{
+		Action->Shutdown();
+	}
+	RegisteredActions.Reset();
+
 	for (TObjectPtr<UUVEditorToolMeshInput> ToolInput : ToolInputObjects)
 	{
 		ToolInput->Shutdown();
@@ -509,29 +661,98 @@ void UUVEditorMode::Exit()
 
 	bIsActive = false;
 
+	UContextObjectStore* ContextStore = GetInteractiveToolsContext()->ToolManager->GetContextObjectStore();
+	for (TWeakObjectPtr<UUVToolContextObject> Context : ContextsToShutdown)
+	{
+		if (ensure(Context.IsValid()))
+		{
+			Context->Shutdown();
+			ContextStore->RemoveContextObject(Context.Get());
+		}
+	}
+
+	GetInteractiveToolsContext()->OnRender.RemoveAll(this);
+	GetInteractiveToolsContext()->OnDrawHUD.RemoveAll(this);
+
 	Super::Exit();
 }
 
-void UUVEditorMode::InitializeContexts(FEditorViewportClient& LivePreviewViewportClient, FAssetEditorModeManager& LivePreviewModeManager, 
+void UUVEditorMode::InitializeAssetEditorContexts(UContextObjectStore& ContextStore,
+	const TArray<TObjectPtr<UObject>>& AssetsIn, const TArray<FTransform>& TransformsIn,
+	FEditorViewportClient& LivePreviewViewportClient, FAssetEditorModeManager& LivePreviewModeManager,
 	UUVToolViewportButtonsAPI& ViewportButtonsAPI, UUVTool2DViewportAPI& UVTool2DViewportAPI)
 {
 	using namespace UVEditorModeLocals;
 
+	UUVToolAssetInputsContext* AssetInputsContext = ContextStore.FindContext<UUVToolAssetInputsContext>();
+	if (!AssetInputsContext)
+	{
+		AssetInputsContext = NewObject<UUVToolAssetInputsContext>();
+		AssetInputsContext->Initialize(AssetsIn, TransformsIn);
+		ContextStore.AddContextObject(AssetInputsContext);
+	}
+
+	UUVToolLivePreviewAPI* LivePreviewAPI = ContextStore.FindContext<UUVToolLivePreviewAPI>();
+	if (!LivePreviewAPI)
+	{
+		LivePreviewAPI = NewObject<UUVToolLivePreviewAPI>();
+		LivePreviewAPI->Initialize(
+			LivePreviewModeManager.GetPreviewScene()->GetWorld(),
+			LivePreviewModeManager.GetInteractiveToolsContext()->InputRouter,
+			[LivePreviewViewportClientPtr = &LivePreviewViewportClient](FViewCameraState& CameraStateOut) {
+				GetCameraState(*LivePreviewViewportClientPtr, CameraStateOut);
+			});
+		ContextStore.AddContextObject(LivePreviewAPI);
+	}
+
+	if (!ContextStore.FindContext<UUVToolViewportButtonsAPI>())
+	{
+		ContextStore.AddContextObject(&ViewportButtonsAPI);
+	}
+
+	if (!ContextStore.FindContext<UUVTool2DViewportAPI>())
+	{
+		ContextStore.AddContextObject(&UVTool2DViewportAPI);
+	}
+}
+
+void UUVEditorMode::InitializeModeContexts()
+{
 	UContextObjectStore* ContextStore = GetInteractiveToolsContext()->ToolManager->GetContextObjectStore();
 
-	LivePreviewWorld = LivePreviewModeManager.GetPreviewScene()->GetWorld();
-	UUVToolLivePreviewAPI* LivePreviewAPI = NewObject<UUVToolLivePreviewAPI>();
-	LivePreviewAPI->Initialize(
-		LivePreviewWorld,
-		LivePreviewModeManager.GetInteractiveToolsContext()->InputRouter,
-		[this, LivePreviewViewportClientPtr = &LivePreviewViewportClient](FViewCameraState& CameraStateOut) {
-			GetCameraState(*LivePreviewViewportClientPtr, CameraStateOut); 
-		});
-	ContextStore->AddContextObject(LivePreviewAPI);
+	UUVToolAssetInputsContext* AssetInputsContext = ContextStore->FindContext<UUVToolAssetInputsContext>();
+	check(AssetInputsContext);
+	AssetInputsContext->GetAssetInputs(OriginalObjectsToEdit, Transforms);
+	ContextsToUpdateOnToolEnd.Add(AssetInputsContext);
+
+	UUVToolLivePreviewAPI* LivePreviewAPI = ContextStore->FindContext<UUVToolLivePreviewAPI>();
+	check(LivePreviewAPI);
+	LivePreviewWorld = LivePreviewAPI->GetLivePreviewWorld();
+	ContextsToUpdateOnToolEnd.Add(LivePreviewAPI);
+
+	UUVToolViewportButtonsAPI* ViewportButtonsAPI = ContextStore->FindContext<UUVToolViewportButtonsAPI>();
+	check(ViewportButtonsAPI);
+	ContextsToUpdateOnToolEnd.Add(ViewportButtonsAPI);
+
+	UUVTool2DViewportAPI* UVTool2DViewportAPI = ContextStore->FindContext<UUVTool2DViewportAPI>();
+	check(UVTool2DViewportAPI);
+	ContextsToUpdateOnToolEnd.Add(UVTool2DViewportAPI);
+
+	// Helper function for adding contexts that our mode creates itself, rather than getting from
+	// the asset editor.
+	auto AddContextObject = [this, ContextStore](UUVToolContextObject* Object)
+	{
+		if (ensure(ContextStore->AddContextObject(Object)))
+		{
+			ContextsToShutdown.Add(Object);
+		}
+		ContextsToUpdateOnToolEnd.Add(Object);
+	};
 
 	UUVToolEmitChangeAPI* EmitChangeAPI = NewObject<UUVToolEmitChangeAPI>();
+	EmitChangeAPI = NewObject<UUVToolEmitChangeAPI>();
 	EmitChangeAPI->Initialize(GetInteractiveToolsContext()->ToolManager);
-	ContextStore->AddContextObject(EmitChangeAPI);
+	AddContextObject(EmitChangeAPI);
 
 	UUVToolAssetAndChannelAPI* AssetAndLayerAPI = NewObject<UUVToolAssetAndChannelAPI>();
 	AssetAndLayerAPI->RequestChannelVisibilityChangeFunc = [this](const TArray<int32>& LayerPerAsset, bool bEmitUndoTransaction) {
@@ -550,27 +771,25 @@ void UUVEditorMode::InitializeContexts(FEditorViewportClient& LivePreviewViewpor
 		}
 		return VisibleLayers;
 	};
-	ContextStore->AddContextObject(AssetAndLayerAPI);
+	AddContextObject(AssetAndLayerAPI);
 
-	ContextStore->AddContextObject(&ViewportButtonsAPI);
+	SelectionAPI = NewObject<UUVToolSelectionAPI>();
+	SelectionAPI->Initialize(GetToolManager(), GetWorld(), 
+		GetInteractiveToolsContext()->InputRouter, LivePreviewAPI, EmitChangeAPI);
+	AddContextObject(SelectionAPI);
 
-	ContextStore->AddContextObject(&UVTool2DViewportAPI);
 }
 
-void UUVEditorMode::InitializeTargets(const TArray<TObjectPtr<UObject>>& AssetsIn,
-	const TArray<FTransform>& TransformsIn)
+void UUVEditorMode::InitializeTargets()
 {
 	using namespace UVEditorModeLocals;
 
-	// InitializeContexts needs to have been called first so that we have the 3d preview world ready.
+	// InitializeModeContexts needs to have been called first so that we have the 3d preview world ready.
 	check(LivePreviewWorld);
 
-	OriginalObjectsToEdit = AssetsIn;
-	Transforms = TransformsIn;
-
 	// Build the tool targets that provide us with 3d dynamic meshes
 	UUVEditorSubsystem* UVSubsystem = GEditor->GetEditorSubsystem<UUVEditorSubsystem>();
-	UVSubsystem->BuildTargets(AssetsIn, GetToolTargetRequirements(), ToolTargets);
+	UVSubsystem->BuildTargets(OriginalObjectsToEdit, GetToolTargetRequirements(), ToolTargets);
 
 	// Collect the 3d dynamic meshes from targets. There will always be one for each asset, and the AssetID
 	// of each asset will be the index into these arrays. Individual input objects (representing a specific
@@ -677,6 +896,9 @@ void UUVEditorMode::InitializeTargets(const TArray<TObjectPtr<UObject>>& AssetsI
 	{
 		UVEditorUDIMProperties->InitializeAssets(ToolInputObjects);
 	}
+
+	// Finish initializing the selection api
+	SelectionAPI->SetTargets(ToolInputObjects);
 }
 
 void UUVEditorMode::EmitToolIndependentObjectChange(UObject* TargetObject, TUniquePtr<FToolCommandChange> Change, const FText& Description)
@@ -969,6 +1191,16 @@ void UUVEditorMode::ChangeInputObjectLayer(int32 AssetID, int32 NewLayerIndex)
 
 void UUVEditorMode::SetDisplayedUVChannels(const TArray<int32>& LayerPerAsset, bool bEmitUndoTransaction)
 {
+	// Deal with any existing selections first.
+	// TODO: We could consider keeping triangle selections, since these can be valid across different layers,
+	// or converting the selections to elements in the applied mesh and back. The gain for this would probably
+	// be minor, so for now we'll just clear selection here.
+	if (SelectionAPI)
+	{
+		SelectionAPI->ClearSelections(true, bEmitUndoTransaction); // broadcast, maybe emit
+	}
+
+	// Now actually swap out the layers
 	for (int32 AssetID = 0; AssetID < ToolInputObjects.Num(); ++AssetID)
 	{
 		if (ToolInputObjects[AssetID]->UVLayerIndex != LayerPerAsset[AssetID])
diff --git a/Engine/Plugins/Experimental/UVEditor/Source/UVEditor/Private/UVEditorModeToolkit.cpp b/Engine/Plugins/Experimental/UVEditor/Source/UVEditor/Private/UVEditorModeToolkit.cpp
--- a/Engine/Plugins/Experimental/UVEditor/Source/UVEditor/Private/UVEditorModeToolkit.cpp
+++ b/Engine/Plugins/Experimental/UVEditor/Source/UVEditor/Private/UVEditorModeToolkit.cpp
@@ -27,45 +27,6 @@
 
 namespace UVEditorModeToolkitLocals
 {
-	/** 
-	 * Support for undoing a tool start in such a way that we go back to the mode's default
-	 * tool on undo.
-	 */
-	class FUVEditorBeginToolChange : public FToolCommandChange
-	{
-	public:
-		virtual void Apply(UObject* Object) override
-		{
-			// Do nothing, since we don't allow a re-do back into a tool
-		}
-
-		virtual void Revert(UObject* Object) override
-		{
-			UUVEditorMode* Mode = Cast<UUVEditorMode>(Object);
-			// Don't really need the check for default tool since we theoretically shouldn't
-			// be issuing this transaction for starting the default tool, but still...
-			if (Mode && !Mode->IsDefaultToolActive())
-			{
-				Mode->GetInteractiveToolsContext()->EndTool(EToolShutdownType::Cancel);
-				Mode->ActivateDefaultTool();
-			}
-		}
-
-		virtual bool HasExpired(UObject* Object) const override
-		{
-			// To not be expired, we must be in some non-default tool.
-			UUVEditorMode* Mode = Cast<UUVEditorMode>(Object);
-			return !(Mode && Mode->GetInteractiveToolsContext() 
-				&& Mode->GetInteractiveToolsContext()->ToolManager
-				&& Mode->GetInteractiveToolsContext()->ToolManager->HasAnyActiveTool() 
-				&& !Mode->IsDefaultToolActive());
-		}
-
-		virtual FString ToString() const override
-		{
-			return TEXT("FUVEditorBeginToolChange");
-		}
-	};
 }
 
 FUVEditorModeToolkit::FUVEditorModeToolkit()
@@ -199,29 +160,27 @@ void FUVEditorModeToolkit::Init(const TSharedPtr<IToolkitHost>& InitToolkitHost,
 				.Visibility_Lambda([this]() { return GetScriptableEditorMode()->GetInteractiveToolsContext()->ActiveToolHasAccept() ? EVisibility::Visible : EVisibility::Collapsed; })
 			]
 
-			// For now we've decided not to use a "Complete" button for complete-style tools, instead requiring
-			// users to just select a different tool. Uncomment the below if we want to put it back.
-			//+ SHorizontalBox::Slot()
-			//.AutoWidth()
-			//.Padding(FMargin(2.0, 0.f, 0.f, 0.f))
-			//[
-			//	SNew(SButton)
-			//	.ButtonStyle(FAppStyle::Get(), "PrimaryButton")
-			//	.TextStyle(FAppStyle::Get(), "DialogButtonText")
-			//	.Text(LOCTEXT("OverlayComplete", "Complete"))
-			//	.ToolTipText(LOCTEXT("OverlayCompleteTooltip", "Exit the active Tool [Enter]"))
-			//	.HAlign(HAlign_Center)
-			//	.OnClicked_Lambda([this]() { 
-			//		GetScriptableEditorMode()->GetInteractiveToolsContext()->EndTool(EToolShutdownType::Completed); 
-			//		Cast<UUVEditorMode>(GetScriptableEditorMode())->ActivateDefaultTool();
-			//		return FReply::Handled(); 
-			//		})
-			//	.IsEnabled_Lambda([this]() {
-			//		UUVEditorMode* Mode = Cast<UUVEditorMode>(GetScriptableEditorMode());
-			//		return GetScriptableEditorMode()->GetInteractiveToolsContext()->CanCompleteActiveTool();
-			//	})
-			//	.Visibility_Lambda([this]() { return GetScriptableEditorMode()->GetInteractiveToolsContext()->CanCompleteActiveTool() ? EVisibility::Visible : EVisibility::Collapsed; })
-			//]
+			+ SHorizontalBox::Slot()
+			.AutoWidth()
+			.Padding(FMargin(2.0, 0.f, 0.f, 0.f))
+			[
+				SNew(SButton)
+				.ButtonStyle(FAppStyle::Get(), "PrimaryButton")
+				.TextStyle(FAppStyle::Get(), "DialogButtonText")
+				.Text(LOCTEXT("OverlayComplete", "Complete"))
+				.ToolTipText(LOCTEXT("OverlayCompleteTooltip", "Exit the active Tool [Enter]"))
+				.HAlign(HAlign_Center)
+				.OnClicked_Lambda([this]() { 
+					GetScriptableEditorMode()->GetInteractiveToolsContext()->EndTool(EToolShutdownType::Completed);
+					Cast<UUVEditorMode>(GetScriptableEditorMode())->ActivateDefaultTool();
+					return FReply::Handled(); 
+					})
+				.IsEnabled_Lambda([this]() {
+					UUVEditorMode* Mode = Cast<UUVEditorMode>(GetScriptableEditorMode());
+					return GetScriptableEditorMode()->GetInteractiveToolsContext()->CanCompleteActiveTool();
+				})
+				.Visibility_Lambda([this]() { return GetScriptableEditorMode()->GetInteractiveToolsContext()->CanCompleteActiveTool() ? EVisibility::Visible : EVisibility::Collapsed; })
+			]
 		]	
 	];
 
@@ -356,8 +315,9 @@ void FUVEditorModeToolkit::InvalidateCachedDetailPanelState(UObject* ChangedObje
 void FUVEditorModeToolkit::OnToolStarted(UInteractiveToolManager* Manager, UInteractiveTool* Tool)
 {
 	FModeToolkit::OnToolStarted(Manager, Tool);
-	
-	UInteractiveTool* CurTool = GetScriptableEditorMode()->GetToolManager()->GetActiveTool(EToolSide::Left);
+
+	UUVEditorMode* Mode = Cast<UUVEditorMode>(GetScriptableEditorMode());
+	UInteractiveTool* CurTool = Mode->GetToolManager()->GetActiveTool(EToolSide::Left);
 	CurTool->OnPropertySetsModified.AddSP(this, &FUVEditorModeToolkit::UpdateActiveToolProperties);
 	CurTool->OnPropertyModifiedDirectlyByTool.AddSP(this, &FUVEditorModeToolkit::InvalidateCachedDetailPanelState);
 
@@ -368,19 +328,10 @@ void FUVEditorModeToolkit::OnToolStarted(UInteractiveToolManager* Manager, UInte
 	FName ActiveToolIconName = ISlateStyle::Join(FUVEditorCommands::Get().GetContextName(), TCHAR_TO_ANSI(*ActiveToolIdentifier));
 	ActiveToolIcon = FUVEditorStyle::Get().GetOptionalBrush(ActiveToolIconName);
 
-	UUVEditorMode* Mode = Cast<UUVEditorMode>(GetScriptableEditorMode());
 	if (!Mode->IsDefaultToolActive())
 	{
-		// Issue a tool start transaction unless we are starting the default tool, because we can't
-		// undo or revert out of the default tool.
-		Mode->GetInteractiveToolsContext()->GetTransactionAPI()->AppendChange(
-			Mode, MakeUnique<UVEditorModeToolkitLocals::FUVEditorBeginToolChange>(), LOCTEXT("ActivateTool", "Activate Tool"));
-
-		if (Mode->GetInteractiveToolsContext()->ActiveToolHasAccept())
-		{
-			// Add the accept/cancel overlay only if the tool has accept/cancel.
-			GetToolkitHost()->AddViewportOverlayWidget(ViewportOverlayWidget.ToSharedRef());
-		}
+		// Add the accept/cancel overlay only if the tool is not the default tool.
+		GetToolkitHost()->AddViewportOverlayWidget(ViewportOverlayWidget.ToSharedRef());
 	}
 }
 
@@ -427,7 +378,10 @@ void FUVEditorModeToolkit::BuildToolPalette(FName PaletteIndex, class FToolBarBu
 
 	if (PaletteIndex == ToolsTabName)
 	{
-		ToolbarBuilder.AddToolBarButton(Commands.BeginSelectTool);
+		ToolbarBuilder.AddToolBarButton(Commands.IslandConformalUnwrapAction);
+		ToolbarBuilder.AddToolBarButton(Commands.SewAction);
+		ToolbarBuilder.AddToolBarButton(Commands.SplitAction);
+
 		ToolbarBuilder.AddToolBarButton(Commands.BeginLayoutTool);
 		ToolbarBuilder.AddToolBarButton(Commands.BeginChannelEditTool);
 		ToolbarBuilder.AddToolBarButton(Commands.BeginSeamTool);
diff --git a/Engine/Plugins/Experimental/UVEditor/Source/UVEditor/Private/UVEditorModule.cpp b/Engine/Plugins/Experimental/UVEditor/Source/UVEditor/Private/UVEditorModule.cpp
--- a/Engine/Plugins/Experimental/UVEditor/Source/UVEditor/Private/UVEditorModule.cpp
+++ b/Engine/Plugins/Experimental/UVEditor/Source/UVEditor/Private/UVEditorModule.cpp
@@ -3,7 +3,6 @@
 #include "UVEditorModule.h"
 
 #include "ContentBrowserMenuContexts.h"
-#include "DetailsCustomizations/UVSelectToolCustomizations.h"
 #include "EditorModeRegistry.h"
 #include "LevelEditorMenuContext.h"
 #include "PropertyEditorModule.h"
@@ -31,9 +30,10 @@ void FUVEditorModule::StartupModule()
 	// Register details view customizations
 	FPropertyEditorModule& PropertyModule = FModuleManager::LoadModuleChecked<FPropertyEditorModule>("PropertyEditor");
 	ClassesToUnregisterOnShutdown.Reset();
-	PropertyModule.RegisterCustomClassLayout(USelectToolActionPropertySet::StaticClass()->GetFName(), 
-		FOnGetDetailCustomizationInstance::CreateStatic(&FUVSelectToolActionPropertySetDetails::MakeInstance));
-	ClassesToUnregisterOnShutdown.Add(USelectToolActionPropertySet::StaticClass()->GetFName());
+	// Customizations get registered like this:
+	//PropertyModule.RegisterCustomClassLayout(USelectToolActionPropertySet::StaticClass()->GetFName(), 
+	//	FOnGetDetailCustomizationInstance::CreateStatic(&FUVSelectToolActionPropertySetDetails::MakeInstance));
+	//ClassesToUnregisterOnShutdown.Add(USelectToolActionPropertySet::StaticClass()->GetFName());
 }
 
 void FUVEditorModule::ShutdownModule()
diff --git a/Engine/Plugins/Experimental/UVEditor/Source/UVEditor/Private/UVEditorToolkit.cpp b/Engine/Plugins/Experimental/UVEditor/Source/UVEditor/Private/UVEditorToolkit.cpp
--- a/Engine/Plugins/Experimental/UVEditor/Source/UVEditor/Private/UVEditorToolkit.cpp
+++ b/Engine/Plugins/Experimental/UVEditor/Source/UVEditor/Private/UVEditorToolkit.cpp
@@ -24,7 +24,7 @@
 #include "UVEditorSubsystem.h"
 #include "UVEditorModule.h"
 #include "UVEditorStyle.h"
-#include "UVToolContextObjects.h"
+#include "ContextObjects/UVToolContextObjects.h"
 #include "UVEditorModeUILayer.h"
 #include "Widgets/Docking/SDockTab.h"
 #include "UVEditorUXSettings.h"
@@ -381,6 +381,20 @@ void FUVEditorToolkit::PostInitAssetEditor()
 	check(PinnedToolkitHost.IsValid());
 	ModeUILayer = MakeShareable(new FUVEditorModeUILayer(PinnedToolkitHost.Get()));
 	ModeUILayer->SetModeMenuCategory( UVEditorMenuCategory );
+
+	TArray<TObjectPtr<UObject>> ObjectsToEdit;
+	OwningAssetEditor->GetObjectsToEdit(ObjectsToEdit);
+
+	// TODO: get these when possible (from level editor selection, for instance), and set them to something reasonable otherwise.
+	TArray<FTransform> ObjectTransforms;
+	ObjectTransforms.SetNum(ObjectsToEdit.Num());
+
+	// This static method call initializes the variety of contexts that UVEditorMode needs to be available in
+	// the context store on Enter() to function properly.
+	UUVEditorMode::InitializeAssetEditorContexts(*EditorModeManager->GetInteractiveToolsContext()->ContextObjectStore, 
+		ObjectsToEdit, ObjectTransforms, *LivePreviewViewportClient, *LivePreviewEditorModeManager, 
+		*ViewportButtonsAPI, *UVTool2DViewportAPI);
+
 	// Currently, aside from setting up all the UI elements, the toolkit also kicks off the UV
 	// editor mode, which is the mode that the editor always works in (things are packaged into
 	// a mode so that they can be moved to another asset editor if necessary).
@@ -391,18 +405,6 @@ void FUVEditorToolkit::PostInitAssetEditor()
 		EditorModeManager->GetActiveScriptableMode(UUVEditorMode::EM_UVEditorModeId));
 	check(UVMode);
 
-	// The mode will need to be able to get to the live preview world, camera, input router, and viewport buttons.
-	UVMode->InitializeContexts(*LivePreviewViewportClient, *LivePreviewEditorModeManager, *ViewportButtonsAPI, *UVTool2DViewportAPI);
-
-	TArray<TObjectPtr<UObject>> ObjectsToEdit;
-	OwningAssetEditor->GetObjectsToEdit(ObjectsToEdit);
-
-	// TODO: get these when possible, set them otherwise.
-	TArray<FTransform> ObjectTransforms;
-	ObjectTransforms.SetNum(ObjectsToEdit.Num());
-
-	UVMode->InitializeTargets(ObjectsToEdit, ObjectTransforms);
-
 	// Regardless of how the user has modified the layout, we're going to make sure that we have
 	// a 2d viewport that will allow our mode to receive ticks.
     // We don't need to invoke the tool palette tab anymore, since this is handled by
diff --git a/Engine/Plugins/Experimental/UVEditor/Source/UVEditor/Public/DetailsCustomizations/UVSelectToolCustomizations.h b/Engine/Plugins/Experimental/UVEditor/Source/UVEditor/Public/DetailsCustomizations/UVSelectToolCustomizations.h
deleted file mode 100644
diff --git a/Engine/Plugins/Experimental/UVEditor/Source/UVEditor/Public/UVEditor2DViewportClient.h b/Engine/Plugins/Experimental/UVEditor/Source/UVEditor/Public/UVEditor2DViewportClient.h
--- a/Engine/Plugins/Experimental/UVEditor/Source/UVEditor/Public/UVEditor2DViewportClient.h
+++ b/Engine/Plugins/Experimental/UVEditor/Source/UVEditor/Public/UVEditor2DViewportClient.h
@@ -7,10 +7,10 @@
 #include "EditorViewportClient.h"
 #include "InputBehaviorSet.h"
 #include "UVEditor2DViewportBehaviorTargets.h" // FUVEditor2DScrollBehaviorTarget, FUVEditor2DMouseWheelZoomBehaviorTarget
-#include "UVToolContextObjects.h" // UUVToolViewportButtonsAPI::ESelectionMode
+#include "ContextObjects/UVToolViewportButtonsAPI.h" // UUVToolViewportButtonsAPI::ESelectionMode
 
-class UUVToolViewportButtonsAPI;
 class UCanvas;
+class UUVTool2DViewportAPI;
 
 /**
  * Client used to display a 2D view of the UV's, implemented by using a perspective viewport with a locked
diff --git a/Engine/Plugins/Experimental/UVEditor/Source/UVEditor/Public/UVEditorCommands.h b/Engine/Plugins/Experimental/UVEditor/Source/UVEditor/Public/UVEditorCommands.h
--- a/Engine/Plugins/Experimental/UVEditor/Source/UVEditor/Public/UVEditorCommands.h
+++ b/Engine/Plugins/Experimental/UVEditor/Source/UVEditor/Public/UVEditorCommands.h
@@ -18,7 +18,6 @@ public:
 	TSharedPtr<FUICommandInfo> OpenUVEditor;
 	TSharedPtr<FUICommandInfo> ApplyChanges;
 
-	TSharedPtr<FUICommandInfo> BeginSelectTool;
 	TSharedPtr<FUICommandInfo> BeginLayoutTool;
 	TSharedPtr<FUICommandInfo> BeginParameterizeMeshTool;
 	TSharedPtr<FUICommandInfo> BeginChannelEditTool;
diff --git a/Engine/Plugins/Experimental/UVEditor/Source/UVEditor/Public/UVEditorMode.h b/Engine/Plugins/Experimental/UVEditor/Source/UVEditor/Public/UVEditorMode.h
--- a/Engine/Plugins/Experimental/UVEditor/Source/UVEditor/Public/UVEditorMode.h
+++ b/Engine/Plugins/Experimental/UVEditor/Source/UVEditor/Public/UVEditorMode.h
@@ -5,7 +5,7 @@
 #include "CoreMinimal.h"
 
 #include "Tools/UEdMode.h"
-#include "ToolTargets/ToolTarget.h"
+#include "ToolTargets/ToolTarget.h" // FToolTargetTypeRequirements
 #include "GeometryBase.h"
 #include "InteractiveTool.h"
 
@@ -13,24 +13,24 @@
 
 PREDECLARE_GEOMETRY(class FDynamicMesh3);
 
-class FEditorViewportClient;
 class FAssetEditorModeManager;
+class FEditorViewportClient;
 class FToolCommandChange;
-class IUVUnwrapDynamicMesh;
+class UContextObjectStore;
+class UInteractiveToolPropertySet;
 class UMeshElementsVisualizer;
-class UPreviewMesh;
-class UToolTarget;
-class UWorld;
-class FUVEditorModeToolkit;
-class UInteractiveToolPropertySet; 
 class UMeshOpPreviewWithBackgroundCompute;
-class UUVToolStateObjectStore;
+class UPreviewMesh;
+class UTexture2D;
 class UUVEditorToolMeshInput;
 class UUVEditorBackgroundPreview;
+class UUVToolAction;
+class UUVToolContextObject;
+class UUVToolSelectionAPI;
 class UUVToolViewportButtonsAPI;
 class UUVTool2DViewportAPI;
 class UUVEditorMode;
-class UTexture2D;
+class UWorld;
 
 /**
  * Visualization settings for the UUVEditorMode's Grid
@@ -145,8 +145,6 @@ public:
 
 	UUVEditorMode();
 
-	void RegisterTools();
-
 	/**
 	 * Gets the tool target requirements for the mode. The resulting targets undergo further processing
 	 * to turn them into the input objects that tools get (since these need preview meshes, etc).
@@ -158,11 +156,13 @@ public:
 	 */
 	static double GetUVMeshScalingFactor();
 
-	// Both initialization functions must be called for things to function properly. InitializeContexts should
-	// be done first so that the 3d preview world is ready for creating meshes in InitializeTargets.
-	void InitializeContexts(FEditorViewportClient& LivePreviewViewportClient, FAssetEditorModeManager& LivePreviewModeManager,
+	/**
+	 * Called by an asset editor so that a created instance of the mode has all the data it needs on Enter() to initialize itself.
+	 */
+	static void InitializeAssetEditorContexts(UContextObjectStore& ContextStore,
+		const TArray<TObjectPtr<UObject>>& AssetsIn, const TArray<FTransform>& TransformsIn,
+		FEditorViewportClient& LivePreviewViewportClient, FAssetEditorModeManager& LivePreviewModeManager,
 		UUVToolViewportButtonsAPI& ViewportButtonsAPI, UUVTool2DViewportAPI& UVTool2DViewportAPI);
-	void InitializeTargets(const TArray<TObjectPtr<UObject>>& AssetsIn, const TArray<FTransform>& TransformsIn);
 
 	// Public for use by undo/redo. Otherwise should use RequestUVChannelChange
 	void ChangeInputObjectLayer(int32 AssetID, int32 NewLayerIndex);
@@ -209,6 +209,9 @@ public:
 	/** @return A settings object suitable for display in a details panel to control UDIM configuration. */
 	UObject* GetUDIMSettingsObject();
 
+	virtual void Render(IToolsContextRenderAPI* RenderAPI) ;
+	virtual void DrawHUD(FCanvas* Canvas, IToolsContextRenderAPI* RenderAPI);
+
 	// UEdMode overrides
 	virtual void Enter() override;
 	virtual bool ShouldToolStartBeAllowed(const FString& ToolIdentifier) const override;
@@ -241,13 +244,20 @@ public:
 	void PopulateUDIMsByAsset(int32 AssetId, TArray<FUDIMSpecifier>& UDIMsOut) const;
 
 protected:
+	UPROPERTY()
+	TArray<TObjectPtr<UUVToolAction>> RegisteredActions;
+
+	void InitializeModeContexts();
+	void InitializeTargets();
+	void RegisterTools();
+	void RegisterActions();
 
 	// UEdMode overrides
 	virtual void CreateToolkit() override;
 	// Not sure whether we need these yet
 	virtual void BindCommands() override;
-	virtual void OnToolStarted(UInteractiveToolManager* Manager, UInteractiveTool* Tool) override {}
-	virtual void OnToolEnded(UInteractiveToolManager* Manager, UInteractiveTool* Tool) override {}
+	virtual void OnToolStarted(UInteractiveToolManager* Manager, UInteractiveTool* Tool) override;
+	virtual void OnToolEnded(UInteractiveToolManager* Manager, UInteractiveTool* Tool) override;
 	
 	void UpdateTriangleMaterialBasedOnBackground(bool IsBackgroundVisible);
 	void UpdatePreviewMaterialBasedOnBackground();
@@ -321,12 +331,18 @@ protected:
 	UPROPERTY()
 	TObjectPtr<UWorld> LivePreviewWorld = nullptr;
 
+	UPROPERTY()
+	TObjectPtr<UUVToolSelectionAPI> SelectionAPI = nullptr;
+
 	/**
 	 * Mode-level property objects (visible or not) that get ticked.
 	 */
 	UPROPERTY()
 	TArray<TObjectPtr<UInteractiveToolPropertySet>> PropertyObjectsToTick;
 
+	TArray<TWeakObjectPtr<UUVToolContextObject>> ContextsToUpdateOnToolEnd;
+	TArray<TWeakObjectPtr<UUVToolContextObject>> ContextsToShutdown;
+
 	bool bIsActive = false;
 	FString DefaultToolIdentifier;
 
diff --git a/Engine/Plugins/Experimental/UVEditor/Source/UVEditorTools/Private/Actions/UVIslandConformalUnwrapAction.cpp b/Engine/Plugins/Experimental/UVEditor/Source/UVEditorTools/Private/Actions/UVIslandConformalUnwrapAction.cpp
new file mode 100644
--- /dev/null
+++ b/Engine/Plugins/Experimental/UVEditor/Source/UVEditorTools/Private/Actions/UVIslandConformalUnwrapAction.cpp
@@ -0,0 +1,159 @@
+// Copyright Epic Games, Inc. All Rights Reserved.
+
+#include "Actions/UVIslandConformalUnwrapAction.h"
+
+#include "Algo/Unique.h"
+#include "ToolTargets/UVEditorToolMeshInput.h"
+#include "Selection/UVToolSelection.h"
+#include "MeshOpPreviewHelpers.h" // UMeshOpPreviewWithBackgroundCompute
+#include "PreviewMesh.h"
+
+#include "Selection/UVToolSelectionAPI.h"
+
+#include "ToolSetupUtil.h"
+#include "ContextObjects/UVToolContextObjects.h"
+#include "Parameterization/DynamicMeshUVEditor.h"
+#include "Selections/MeshConnectedComponents.h"
+
+#define LOCTEXT_NAMESPACE "UUVIslandConformalUnwrapAction"
+
+using namespace UE::Geometry;
+
+namespace UVIslandConformalUnwrapActionLocals
+{
+	bool GatherIslandTids(const FUVToolSelection& SelectionIn, 
+		TArray<int32>& ConcatenatedIslandTidsOut, TArray<int32>& IslandStartIndicesOut, int32& MaxIslandSizeOut) {
+		TRACE_CPUPROFILER_EVENT_SCOPE(UVIslandConformalUnwrapAction_GatherIslandTids);
+
+		ConcatenatedIslandTidsOut.Reset();
+		IslandStartIndicesOut.Reset();
+		MaxIslandSizeOut = 0;
+
+		FDynamicMesh3& AppliedMesh = *(SelectionIn.Target->AppliedCanonical);
+		FDynamicMeshUVOverlay* UseOverlay = AppliedMesh.Attributes()->GetUVLayer(SelectionIn.Target->UVLayerIndex);
+
+		FMeshConnectedComponents ConnectedComponents(&AppliedMesh);
+
+		ConnectedComponents.FindConnectedTriangles([&](int32 Triangle0, int32 Triangle1) {
+			return UseOverlay->AreTrianglesConnected(Triangle0, Triangle1);
+			});
+
+		int32 NumComponents = ConnectedComponents.Num();
+		IslandStartIndicesOut.Add(0);
+		for (int32 k = 0; k < NumComponents; ++k)
+		{
+			TSet<int32> ComponentTris = TSet<int32>(ConnectedComponents[k].Indices);
+			TSet<int32> SelectedComponentTris = ComponentTris.Intersect(SelectionIn.SelectedIDs);
+			if (SelectedComponentTris.Num() == ComponentTris.Num())
+			{
+				ConcatenatedIslandTidsOut.Append(SelectedComponentTris.Array());
+				IslandStartIndicesOut.Add(IslandStartIndicesOut.Last(0) + SelectedComponentTris.Num());
+				MaxIslandSizeOut = FMath::Max(MaxIslandSizeOut, SelectedComponentTris.Num());
+			}
+		}
+
+		return true;
+	}
+}
+
+bool UUVIslandConformalUnwrapAction::CanExecuteAction() const
+{
+	return SelectionAPI->HaveSelections() 
+		&& SelectionAPI->GetSelectionsType() == FUVToolSelection::EType::Triangle;
+}
+
+bool UUVIslandConformalUnwrapAction::ExecuteAction()
+{
+	using namespace UVIslandConformalUnwrapActionLocals;
+
+	TRACE_CPUPROFILER_EVENT_SCOPE(UVIslandConformalUnwrapAction_ApplyAction);
+
+	for (const FUVToolSelection& Selection : SelectionAPI->GetSelections())
+	{
+		if (!ensure(Selection.Target.IsValid() && Selection.Target->IsValid()))
+		{
+			continue;
+		}
+
+		TArray<int32> ConcatenatedIslandTids;
+		TArray<int32> IslandStartIndices;
+		int32 MaxIslandSize;
+
+		GatherIslandTids(Selection, ConcatenatedIslandTids, IslandStartIndices, MaxIslandSize);
+
+		UUVEditorToolMeshInput* Target = Selection.Target.Get();
+
+		FDynamicMesh3& MeshToUnwrap = *(Target->AppliedCanonical);
+		FDynamicMeshUVOverlay* UVOverlay = MeshToUnwrap.Attributes()->GetUVLayer(Target->UVLayerIndex);
+		FDynamicMeshUVEditor UVEditor(&MeshToUnwrap, Target->UVLayerIndex, true);
+
+		TArray<int32> ComponentTris;
+		ComponentTris.Reserve(MaxIslandSize);
+		int32 NumComponents = IslandStartIndices.Num() - 1;
+		TArray<bool> bComponentSolved;
+		bComponentSolved.Init(false, NumComponents);
+		int32 SuccessCount = 0;
+		TSet<int32> VidSet;
+		for (int32 k = 0; k < NumComponents; ++k)
+		{
+			ComponentTris.Reset(0);
+			int32 StartIndex = IslandStartIndices[k];
+			int32 Count = IslandStartIndices[k + 1] - StartIndex;
+			if (Count == 0)
+			{
+				continue; // Skip everything if there are no Tids for this island component
+			}
+			ComponentTris.Append(&ConcatenatedIslandTids[StartIndex], Count);
+
+			FAxisAlignedBox2f UVBounds = FAxisAlignedBox2f::Empty();
+			for (int32 tid : ComponentTris)
+			{
+				if (UVOverlay->IsSetTriangle(tid))
+				{
+					FIndex3i UVTri = UVOverlay->GetTriangle(tid);
+					FVector2f U = UVOverlay->GetElement(UVTri.A);
+					FVector2f V = UVOverlay->GetElement(UVTri.B);
+					FVector2f W = UVOverlay->GetElement(UVTri.C);
+					UVBounds.Contain(U); UVBounds.Contain(V); UVBounds.Contain(W);
+				}
+			}
+
+
+			bComponentSolved[k] = UVEditor.SetTriangleUVsFromFreeBoundaryConformal(ComponentTris, true);
+			if (bComponentSolved[k])
+			{
+				UVEditor.ScaleUVAreaToBoundingBox(ComponentTris, UVBounds, true, true);
+			}
+
+			for (int32 Tid : ComponentTris)
+			{
+				FIndex3i TriVids = MeshToUnwrap.Attributes()->GetUVLayer(Target->UVLayerIndex)->GetTriangle(Tid);
+				for (int i = 0; i < 3; ++i)
+				{
+					VidSet.Add(TriVids[i]);
+				}
+			}
+		}
+
+		if (VidSet.IsEmpty())
+		{
+			continue;
+		}
+
+		FDynamicMeshChangeTracker ChangeTracker(Target->UnwrapCanonical.Get());
+		ChangeTracker.BeginChange();
+		ChangeTracker.SaveTriangles(ConcatenatedIslandTids, true);
+
+		const TArray<int32> AllVids = VidSet.Array();
+		Target->UpdateAllFromAppliedCanonical(&AllVids, UUVEditorToolMeshInput::NONE_CHANGED_ARG, &ConcatenatedIslandTids);
+
+		const FText TransactionName(LOCTEXT("ConformalUnwrapCompleteTransactionName", "Conformal Unwrap Islands"));
+		EmitChangeAPI->EmitToolIndependentUnwrapCanonicalChange(Target,
+			ChangeTracker.EndChange(), TransactionName);
+	}
+
+	// TODO: Should maybe give some kind of error message if there were no fully selected components?
+	return true;
+}
+
+#undef LOCTEXT_NAMESPACE
\ No newline at end of file
diff --git a/Engine/Plugins/Experimental/UVEditor/Source/UVEditorTools/Private/Actions/UVSeamSewAction.cpp b/Engine/Plugins/Experimental/UVEditor/Source/UVEditorTools/Private/Actions/UVSeamSewAction.cpp
new file mode 100644
--- /dev/null
+++ b/Engine/Plugins/Experimental/UVEditor/Source/UVEditorTools/Private/Actions/UVSeamSewAction.cpp
@@ -0,0 +1,221 @@
+// Copyright Epic Games, Inc. All Rights Reserved.
+
+#include "Actions/UVSeamSewAction.h"
+
+#include "MeshOpPreviewHelpers.h"
+#include "ToolTargets/UVEditorToolMeshInput.h"
+#include "Selection/UVToolSelection.h"
+#include "Selection/UVToolSelectionAPI.h"
+#include "ContextObjects/UVToolContextObjects.h"
+#include "UVEditorUXSettings.h"
+
+
+#define LOCTEXT_NAMESPACE "UUVSeamSewAction"
+
+using namespace UE::Geometry;
+
+int32 UUVSeamSewAction::FindSewEdgeOppositePairing(const FDynamicMesh3& UnwrapMesh, const FDynamicMesh3& AppliedMesh, 
+	int32 UVLayerIndex, int32 UnwrapEid, bool& bWouldPreferReverseOut)
+{
+	// Given an edge id on the unwrap mesh, determine its opposite edge suitable for UV sewing elsewhere on the unwrap mesh
+
+	if (!UnwrapMesh.IsBoundaryEdge(UnwrapEid))
+	{
+		return IndexConstants::InvalidID;
+	}
+
+	int32 StartUnwrapTid = UnwrapMesh.GetEdgeT(UnwrapEid).A;
+	int StartEdgeSubIndex = IndexUtil::FindTriIndex(UnwrapEid, UnwrapMesh.GetTriEdges(StartUnwrapTid));
+
+	int32 AppliedEid = AppliedMesh.GetTriEdge(StartUnwrapTid, StartEdgeSubIndex);
+	FIndex2i AppliedTids = AppliedMesh.GetEdgeT(AppliedEid);
+	int32 EndUnwrapTid = (AppliedTids.A == StartUnwrapTid) ? AppliedTids.B : AppliedTids.A;
+
+	if (EndUnwrapTid == IndexConstants::InvalidID || !UnwrapMesh.IsTriangle(EndUnwrapTid))
+	{
+		// Either this edge is a boundary edge in the original mesh, or the other triangle in the original
+		// mesh has unset UV's.
+		return IndexConstants::InvalidID;
+	}
+
+	int32 EndEdgeSubIndex = IndexUtil::FindTriIndex(AppliedEid, AppliedMesh.GetTriEdges(EndUnwrapTid));
+	int32 PairedEid = UnwrapMesh.GetTriEdge(EndUnwrapTid, EndEdgeSubIndex);
+
+	bWouldPreferReverseOut = StartUnwrapTid > EndUnwrapTid;
+
+	return PairedEid;
+}
+
+bool UUVSeamSewAction::CanExecuteAction() const
+{
+	return SelectionAPI->HaveSelections()
+		&& SelectionAPI->GetSelectionsType() == FUVToolSelection::EType::Edge;
+}
+
+bool UUVSeamSewAction::ExecuteAction()
+{
+	TRACE_CPUPROFILER_EVENT_SCOPE(UVSeamSewAction_ExecuteAction);
+	
+	const FText TransactionName(LOCTEXT("SewCompleteTransactionName", "Sew Edges"));
+
+	// To properly undo, we need to emit the selection change before we emit the mesh change.
+	// As long as we change at least a single mesh, we'll clear the entire selection, then set
+	// it to the new selection at the end.
+	bool bAtLeastOneValidCandidate = false;
+	TArray<FUVToolSelection> NewSelections;
+	TArray<FUVToolSelection> StartSelections = SelectionAPI->GetSelections();
+	for (const FUVToolSelection& Selection : StartSelections)
+	{
+		if (!ensure(Selection.Target.IsValid() && Selection.Target->IsValid()))
+		{
+			continue;
+		}
+
+		UUVEditorToolMeshInput* Target = Selection.Target.Get();
+		FDynamicMesh3& MeshToSew = *(Target->UnwrapCanonical);
+		FUVToolSelection NewSelection;
+		NewSelection.Target = Selection.Target;
+		NewSelection.Type = Selection.Type;
+
+		// Gather up all the edge pairs to sew.
+		TArray<FIndex2i> EdgeSewPairs;
+		for (int32 Eid : Selection.SelectedIDs)
+		{
+			bool bWouldPreferReverse = false;
+			int32 PairedEid = FindSewEdgeOppositePairing(
+				MeshToSew, *Target->AppliedCanonical,
+				Target->UVLayerIndex, Eid, bWouldPreferReverse);
+
+			if (PairedEid == IndexConstants::InvalidID)
+			{
+				NewSelection.SelectedIDs.Add(Eid); // unpairable edges should stay selected.
+				continue;
+			}
+
+			// If the paired edge is also selected, use a consistent prefernce for which goes to which
+			// (i.e. skip if we would prefer the welding to happen starting with the other edge).
+			if (Selection.SelectedIDs.Contains(PairedEid) && bWouldPreferReverse)
+			{
+				continue;
+			}
+
+			EdgeSewPairs.Emplace(Eid, PairedEid);
+		}
+
+		if (EdgeSewPairs.Num() == 0)
+		{
+			continue;
+		}
+
+		if (!bAtLeastOneValidCandidate)
+		{
+			EmitChangeAPI->BeginUndoTransaction(TransactionName);
+			SelectionAPI->ClearSelections(false, true); // Don't broadcast, do emit
+			bAtLeastOneValidCandidate = true;
+		}
+
+		// Gather up all the triangles that may be changed. These are not just the triangles on either side
+		// of the edge but triangles in one rings of the vertices.
+		TSet<int32> ChangedTidSet;
+		for (const FIndex2i& EidPair : EdgeSewPairs)
+		{
+			for (int PairIndex = 0; PairIndex < 2; ++PairIndex)
+			{
+				FIndex2i Vids = MeshToSew.GetEdgeV(EidPair[PairIndex]);
+				for (int VidIndex = 0; VidIndex < 2; ++VidIndex)
+				{
+					TArray<int32> Tids;
+					MeshToSew.GetVtxTriangles(Vids[VidIndex], Tids);
+					ChangedTidSet.Append(Tids);
+				}
+			}
+		}
+
+		TArray<int32> ChangedTids = ChangedTidSet.Array();
+
+		FDynamicMeshChangeTracker ChangeTracker(Target->UnwrapCanonical.Get());
+		ChangeTracker.BeginChange();
+		ChangeTracker.SaveTriangles(ChangedTids, true);
+
+		// Note that currently, we don't really need to gather up the kept verts and call them out for an
+		// update, since no vert locations should have changed. But we'll do it anyway in case any of the
+		// code changes in some way that moves the remaining verts (for instance, to some halfway point).
+		TSet<int32> KeptVidsSet;
+
+		// Perform the actual weld ops
+		for (FIndex2i EdgePair : EdgeSewPairs)
+		{
+			// An edge may have already been merged while merging an adjacent one, so make sure the edges still
+			// exist. If one doesn't exist, it's pair should still be addable to selection.
+			if (!MeshToSew.IsEdge(EdgePair.A))
+			{
+				if (ensure(MeshToSew.IsEdge(EdgePair.B)))
+				{
+					NewSelection.SelectedIDs.Add(EdgePair.B);
+				}
+				continue;
+			}
+			if (!MeshToSew.IsEdge(EdgePair.B))
+			{
+				if (ensure(MeshToSew.IsEdge(EdgePair.A)))
+				{
+					NewSelection.SelectedIDs.Add(EdgePair.A);
+				}
+				continue;
+			}
+
+			FDynamicMesh3::FMergeEdgesInfo MergeInfo;
+			EMeshResult Result = MeshToSew.MergeEdges(EdgePair[1], EdgePair[0], MergeInfo, false);
+			if (Result == EMeshResult::Ok)
+			{
+				for (int i = 0; i < 2; ++i)
+				{
+					KeptVidsSet.Add(MergeInfo.KeptVerts[i]);
+					NewSelection.SelectedIDs.Add(MergeInfo.KeptEdge);
+				}
+			}
+			else
+			{
+				UE_LOG(LogGeometry, Warning, TEXT("Failed to sew edge pair %d / %d. Failed with code %d"), EdgePair[0], EdgePair[1], Result);
+			}
+		}
+		checkSlow(MeshToSew.CheckValidity(FDynamicMesh3::FValidityOptions(true, true))); // Allow nonmanifold verts and reverse orientation
+
+		TArray<int32> RemainingVids;
+		for (int32 KeptVid : KeptVidsSet)
+		{
+			// We have to do this check because we may have performed multiple merge actions, so a "kept vert"
+			// from one action may have ended up getting removed in a later one.
+			if (MeshToSew.IsVertex(KeptVid))
+			{
+				RemainingVids.Add(KeptVid);
+			}
+		}
+
+		Target->UpdateUnwrapCanonicalOverlayFromPositions(&RemainingVids, &ChangedTids);
+		Target->UpdateAllFromUnwrapCanonical(&RemainingVids, &ChangedTids, &ChangedTids);
+		checkSlow(MeshToSew.IsSameAs(*Target->UnwrapPreview->PreviewMesh->GetMesh(), FDynamicMesh3::FSameAsOptions()));
+
+		EmitChangeAPI->EmitToolIndependentUnwrapCanonicalChange(Target,
+			ChangeTracker.EndChange(), TransactionName);
+
+		if (!NewSelection.IsEmpty())
+		{
+			NewSelections.Add(MoveTemp(NewSelection));
+		}
+	}
+
+	if (!bAtLeastOneValidCandidate)
+	{
+		GetToolManager()->DisplayMessage(
+			LOCTEXT("SewErrorSelectionNotBoundary", "Cannot sew UVs. No viable sew candidate edges selected."),
+			EToolMessageLevel::UserWarning);
+		return false;
+	}
+
+	SelectionAPI->SetSelections(NewSelections, true, true); // broadcast, emit
+	EmitChangeAPI->EndUndoTransaction();
+	return true;
+}
+
+#undef LOCTEXT_NAMESPACE
\ No newline at end of file
diff --git a/Engine/Plugins/Experimental/UVEditor/Source/UVEditorTools/Private/Actions/UVSplitAction.cpp b/Engine/Plugins/Experimental/UVEditor/Source/UVEditorTools/Private/Actions/UVSplitAction.cpp
new file mode 100644
--- /dev/null
+++ b/Engine/Plugins/Experimental/UVEditor/Source/UVEditorTools/Private/Actions/UVSplitAction.cpp
@@ -0,0 +1,210 @@
+// Copyright Epic Games, Inc. All Rights Reserved.
+
+#include "Actions/UVSplitAction.h"
+
+#include "ContextObjects/UVToolContextObjects.h"
+#include "DynamicMesh/DynamicMesh3.h"
+#include "DynamicMesh/DynamicMeshChangeTracker.h"
+#include "Parameterization/DynamicMeshUVEditor.h"
+#include "Selection/UVToolSelection.h"
+#include "Selection/UVToolSelectionAPI.h"
+#include "ToolTargets/UVEditorToolMeshInput.h"
+
+
+#define LOCTEXT_NAMESPACE "UUVSplitAction"
+
+using namespace UE::Geometry;
+
+namespace UVSplitActionLocals
+{
+	bool ApplySplitEdges(const FUVToolSelection& Selection, FUVToolSelection& NewSelectionOut, 
+		UUVToolEmitChangeAPI& EmitChangeAPI, const FText& TransactionName)
+	{
+		UUVEditorToolMeshInput* Target = Selection.Target.Get();
+		
+		if (!ensure(Target && Target->IsValid()))
+		{
+			return false;
+		}
+
+		// Gather up the corresponding edge IDs in the applied (3d) mesh.
+		TSet<int32> AppliedEidSet;
+		for (int32 Eid : Selection.SelectedIDs)
+		{
+			// Note that we don't check whether edges are already boundary edges because we allow such edges
+			// to be selected for splitting of any attached bowties.
+
+			FIndex2i EdgeUnwrapVids = Target->UnwrapCanonical->GetEdgeV(Eid);
+
+			int32 AppliedEid = Target->AppliedCanonical->FindEdge(
+				Target->UnwrapVidToAppliedVid(EdgeUnwrapVids.A),
+				Target->UnwrapVidToAppliedVid(EdgeUnwrapVids.B));
+
+			if (ensure(AppliedEid != IndexConstants::InvalidID))
+			{
+				AppliedEidSet.Add(AppliedEid);
+			}
+		}
+
+		// Perform the cut in the overlay, but don't propagate to unwrap yet because we'll need to prep for undo
+		FUVEditResult UVEditResult;
+		FDynamicMeshUVEditor UVEditor(Target->AppliedCanonical.Get(),
+			Target->UVLayerIndex, false);
+		UVEditor.CreateSeamsAtEdges(AppliedEidSet, &UVEditResult);
+
+		// Figure out the triangles that need to be saved in the unwrap for undo
+		TSet<int32> TidSet;
+		for (int32 UnwrapVid : UVEditResult.NewUVElements)
+		{
+			TArray<int32> VertTids;
+			Target->AppliedCanonical->GetVtxTriangles(Target->UnwrapVidToAppliedVid(UnwrapVid), VertTids);
+			TidSet.Append(VertTids);
+		}
+
+		FDynamicMeshChangeTracker ChangeTracker(Target->UnwrapCanonical.Get());
+		ChangeTracker.BeginChange();
+		ChangeTracker.SaveTriangles(TidSet, true);
+
+		// Perform the update
+		TArray<int32> AppliedTids = TidSet.Array();
+		Target->UpdateAllFromAppliedCanonical(&UVEditResult.NewUVElements, &AppliedTids, &AppliedTids);
+
+		// Gather up the corresponding eids in the unwrap. We select only one of the newly
+		// created border edges because this turns out to be very convenient for splitting
+		// and then moving the edge. This is slightly awkward in some edge cases, namely if
+		// we select an existing border edge, we end up with the other border edge selected
+		// at the end, but dealing with those is unlikely to be worth the code complexity.
+		TSet<int32> UnwrapEidSet;
+		for (int32 AppliedEid : AppliedEidSet)
+		{
+			int32 Tid = Target->AppliedCanonical->GetEdgeT(AppliedEid).A;
+			FIndex3i TriAppliedEids = Target->AppliedCanonical->GetTriEdges(Tid);
+			int32 EdgeSubIndex = IndexUtil::FindTriIndex(AppliedEid, TriAppliedEids);
+
+			FIndex3i TriUnwrapEids = Target->UnwrapCanonical->GetTriEdges(Tid);
+
+			if (ensure(EdgeSubIndex != IndexConstants::InvalidID))
+			{
+				UnwrapEidSet.Add(TriUnwrapEids[EdgeSubIndex]);
+			}
+		}
+
+		NewSelectionOut.Target = Selection.Target;
+		NewSelectionOut.Type = Selection.Type;
+		NewSelectionOut.SelectedIDs = MoveTemp(UnwrapEidSet);
+
+		// Emit update transaction
+		EmitChangeAPI.EmitToolIndependentUnwrapCanonicalChange(
+			Target, ChangeTracker.EndChange(), TransactionName);
+
+		return true;
+	}
+
+	bool ApplySplitBowtieVertices(const FUVToolSelection& Selection, FUVToolSelection& NewSelectionOut, 
+		UUVToolEmitChangeAPI& EmitChangeAPI, const FText& TransactionName)
+	{
+		UUVEditorToolMeshInput* Target = Selection.Target.Get();
+
+		if (!ensure(Target && Target->IsValid()))
+		{
+			return false;
+		}
+
+		// Gather the corresponding vert ID's in the applied mesh
+		TSet<int32> AppliedVidSet;
+		for (int32 UnwrapVid : Selection.SelectedIDs)
+		{
+			int32 AppliedVid = Target->UnwrapVidToAppliedVid(UnwrapVid);
+			AppliedVidSet.Add(AppliedVid);
+		}
+
+		// Split any bowties in the applied mesh overlay
+		TArray<int32> NewUVElements;
+		FDynamicMeshUVOverlay* Overlay = Target->AppliedCanonical->Attributes()->GetUVLayer(Target->UVLayerIndex);
+		for (int32 Vid : AppliedVidSet)
+		{
+			Overlay->SplitBowtiesAtVertex(Vid, &NewUVElements);
+		}
+
+		// Prep for undo transaction
+		TSet<int32> TidSet;
+		for (int32 UnwrapVid : NewUVElements)
+		{
+			TArray<int32> VertTids;
+			Target->AppliedCanonical->GetVtxTriangles(Target->UnwrapVidToAppliedVid(UnwrapVid), VertTids);
+			TidSet.Append(VertTids);
+		}
+
+		FDynamicMeshChangeTracker ChangeTracker(Target->UnwrapCanonical.Get());
+		ChangeTracker.BeginChange();
+		ChangeTracker.SaveTriangles(TidSet, true);
+
+		// Perform the update
+		TArray<int32> AppliedTids = TidSet.Array();
+		Target->UpdateAllFromAppliedCanonical(&NewUVElements, &AppliedTids, &AppliedTids);
+
+		// Emit update transaction
+		EmitChangeAPI.EmitToolIndependentUnwrapCanonicalChange(
+			Target, ChangeTracker.EndChange(), TransactionName);
+
+		// Set up the new selection to include the new elements
+		NewSelectionOut = Selection;
+		NewSelectionOut.SelectedIDs.Append(NewUVElements);
+
+		return true;
+	}
+}
+
+bool UUVSplitAction::CanExecuteAction() const
+{
+	return SelectionAPI->HaveSelections()
+		&& (SelectionAPI->GetSelectionsType() == FUVToolSelection::EType::Vertex
+			|| SelectionAPI->GetSelectionsType() == FUVToolSelection::EType::Edge);
+}
+
+bool UUVSplitAction::ExecuteAction()
+{
+	using namespace UVSplitActionLocals;
+
+	if (!ensure(CanExecuteAction()))
+	{
+		return false;
+	}
+
+	bool bAllSucceeded = true;
+	TArray<FUVToolSelection> Selections = SelectionAPI->GetSelections();
+	FUVToolSelection::EType SelectionType = SelectionAPI->GetSelectionsType();
+	TArray<FUVToolSelection> NewSelections;
+
+	const FText TransactionName = (SelectionType == FUVToolSelection::EType::Edge) ?
+		LOCTEXT("SplitEdgesTransactionName", "Split Edges")
+		: LOCTEXT("SplitBowtieVerticesTransactionName", "Split Bowties");
+	EmitChangeAPI->BeginUndoTransaction(TransactionName);
+
+	SelectionAPI->ClearSelections(false, true); // don't broadcast, do emit
+	for (const FUVToolSelection& Selection : Selections)
+	{
+		FUVToolSelection NewSelection;
+		if (Selection.Type == FUVToolSelection::EType::Edge)
+		{
+			bAllSucceeded &= ApplySplitEdges(Selection, NewSelection, *EmitChangeAPI, TransactionName);
+		}
+		else if (Selection.Type == FUVToolSelection::EType::Vertex)
+		{
+			bAllSucceeded &= ApplySplitBowtieVertices(Selection, NewSelection, *EmitChangeAPI, TransactionName);
+		}
+
+		if (!NewSelection.IsEmpty())
+		{
+			NewSelections.Add(NewSelection);
+		}
+	}
+
+	SelectionAPI->SetSelections(NewSelections, true, true); // broadcast and emit.
+
+	EmitChangeAPI->EndUndoTransaction();
+
+	return bAllSucceeded;
+}
+
+#undef LOCTEXT_NAMESPACE
\ No newline at end of file
diff --git a/Engine/Plugins/Experimental/UVEditor/Source/UVEditorTools/Private/Actions/UVToolAction.cpp b/Engine/Plugins/Experimental/UVEditor/Source/UVEditorTools/Private/Actions/UVToolAction.cpp
new file mode 100644
--- /dev/null
+++ b/Engine/Plugins/Experimental/UVEditor/Source/UVEditorTools/Private/Actions/UVToolAction.cpp
@@ -0,0 +1,28 @@
+// Copyright Epic Games, Inc. All Rights Reserved.
+
+#include "Actions/UVToolAction.h"
+
+#include "ContextObjectStore.h"
+#include "Selection/UVToolSelectionAPI.h"
+
+#define LOCTEXT_NAMESPACE "UUVToolAction"
+
+using namespace UE::Geometry;
+
+void UUVToolAction::Setup(UInteractiveToolManager* ToolManagerIn)
+{
+	ToolManager = ToolManagerIn;
+
+	UContextObjectStore* ContextStore = GetToolManager()->GetContextObjectStore();
+	SelectionAPI = ContextStore->FindContext<UUVToolSelectionAPI>();
+	checkSlow(SelectionAPI);
+	EmitChangeAPI = ContextStore->FindContext<UUVToolEmitChangeAPI>();
+}
+
+void UUVToolAction::Shutdown()
+{
+	SelectionAPI = nullptr;
+	EmitChangeAPI = nullptr;
+}
+
+#undef LOCTEXT_NAMESPACE
\ No newline at end of file
diff --git a/Engine/Plugins/Experimental/UVEditor/Source/UVEditorTools/Private/ContextObjects/UVToolContextObjects.cpp b/Engine/Plugins/Experimental/UVEditor/Source/UVEditorTools/Private/ContextObjects/UVToolContextObjects.cpp
new file mode 100644
--- /dev/null
+++ b/Engine/Plugins/Experimental/UVEditor/Source/UVEditorTools/Private/ContextObjects/UVToolContextObjects.cpp
@@ -0,0 +1,175 @@
+// Copyright Epic Games, Inc. All Rights Reserved.
+
+#include "ContextObjects/UVToolContextObjects.h"
+
+#include "DynamicMesh/MeshIndexUtil.h"
+#include "Engine/World.h"
+#include "InputRouter.h" // Need to define this and UWorld so weak pointers know they are UThings
+#include "InteractiveToolManager.h"
+#include "ToolTargets/UVEditorToolMeshInput.h"
+#include "UDIMUtilities.h"
+
+using namespace UE::Geometry;
+
+namespace UVToolContextObjectLocals
+{
+
+	/**
+	 * A wrapper change that applies a given change to the unwrap canonical mesh of an input, and uses that
+	 * to update the other views. Causes a broadcast of OnCanonicalModified.
+	 */
+	class  FUVEditorMeshChange : public FToolCommandChange
+	{
+	public:
+		FUVEditorMeshChange(UUVEditorToolMeshInput* UVToolInputObjectIn, TUniquePtr<FDynamicMeshChange> UnwrapCanonicalMeshChangeIn)
+			: UVToolInputObject(UVToolInputObjectIn)
+			, UnwrapCanonicalMeshChange(MoveTemp(UnwrapCanonicalMeshChangeIn))
+		{
+			ensure(UVToolInputObjectIn);
+			ensure(UnwrapCanonicalMeshChange);
+		};
+
+		virtual void Apply(UObject* Object) override
+		{
+			UnwrapCanonicalMeshChange->Apply(UVToolInputObject->UnwrapCanonical.Get(), false);
+			UVToolInputObject->UpdateFromCanonicalUnwrapUsingMeshChange(*UnwrapCanonicalMeshChange);
+		}
+
+		virtual void Revert(UObject* Object) override
+		{
+			UnwrapCanonicalMeshChange->Apply(UVToolInputObject->UnwrapCanonical.Get(), true);
+			UVToolInputObject->UpdateFromCanonicalUnwrapUsingMeshChange(*UnwrapCanonicalMeshChange);
+		}
+
+		virtual bool HasExpired(UObject* Object) const override
+		{
+			return !(UVToolInputObject.IsValid() && UVToolInputObject->IsValid() && UnwrapCanonicalMeshChange);
+		}
+
+
+		virtual FString ToString() const override
+		{
+			return TEXT("FUVEditorMeshChange");
+		}
+
+	protected:
+		TWeakObjectPtr<UUVEditorToolMeshInput> UVToolInputObject;
+		TUniquePtr<FDynamicMeshChange> UnwrapCanonicalMeshChange;
+	};
+}
+
+void UUVToolEmitChangeAPI::EmitToolIndependentChange(UObject* TargetObject, TUniquePtr<FToolCommandChange> Change, const FText& Description)
+{
+	ToolManager->GetContextTransactionsAPI()->AppendChange(TargetObject, MoveTemp(Change), Description);
+}
+
+void UUVToolEmitChangeAPI::EmitToolIndependentUnwrapCanonicalChange(UUVEditorToolMeshInput* InputObject, 
+	TUniquePtr<FDynamicMeshChange> UnwrapCanonicalMeshChange, const FText& Description)
+{
+	ToolManager->GetContextTransactionsAPI()->AppendChange(InputObject, 
+		MakeUnique<UVToolContextObjectLocals::FUVEditorMeshChange>(InputObject, MoveTemp(UnwrapCanonicalMeshChange)), Description);
+}
+
+void UUVToolEmitChangeAPI::EmitToolDependentChange(UObject* TargetObject, TUniquePtr<FToolCommandChange> Change, const FText& Description)
+{
+	// This should wrap the change in the proper wrapper that will expire it when the tool changes
+	ToolManager->EmitObjectChange(TargetObject, MoveTemp(Change), Description);
+}
+
+void UUVToolLivePreviewAPI::Initialize(UWorld* WorldIn, UInputRouter* RouterIn,
+	TUniqueFunction<void(FViewCameraState& CameraStateOut)> GetLivePreviewCameraStateFuncIn)
+{
+	World = WorldIn;
+	InputRouter = RouterIn;
+	GetLivePreviewCameraStateFunc = MoveTemp(GetLivePreviewCameraStateFuncIn);
+}
+
+int32 FUDIMBlock::BlockU() const
+{
+	int32 BlockU, BlockV;
+	UE::TextureUtilitiesCommon::ExtractUDIMCoordinates(UDIM, BlockU, BlockV);
+	return BlockU;
+}
+
+int32 FUDIMBlock::BlockV() const
+{
+	int32 BlockU, BlockV;
+	UE::TextureUtilitiesCommon::ExtractUDIMCoordinates(UDIM, BlockU, BlockV);
+	return BlockV;
+}
+
+void FUDIMBlock::SetFromBlocks(int32 BlockU, int32 BlockV)
+{
+	UDIM = UE::TextureUtilitiesCommon::GetUDIMIndex(BlockU, BlockV);
+}
+
+void UUVToolAABBTreeStorage::Set(FDynamicMesh3* MeshKey, TSharedPtr<FDynamicMeshAABBTree3> Tree, 
+	UUVEditorToolMeshInput* InputObject)
+{
+	AABBTreeStorage.Add(MeshKey, TreeInputObjectPair(Tree, InputObject));
+	InputObject->OnCanonicalModified.AddWeakLambda(this, [MeshKey, Tree]
+	(UUVEditorToolMeshInput* InputObject, const UUVEditorToolMeshInput::FCanonicalModifiedInfo& Info) {
+		if (InputObject->UnwrapCanonical.Get() == MeshKey
+			|| (Info.bAppliedMeshShapeChanged && InputObject->AppliedCanonical.Get() == MeshKey))
+		{
+			Tree->Build();
+		}
+	});
+}
+
+TSharedPtr<FDynamicMeshAABBTree3> UUVToolAABBTreeStorage::Get(FDynamicMesh3* MeshKey)
+{
+	TreeInputObjectPair* FoundPtr = AABBTreeStorage.Find(MeshKey);
+	return FoundPtr ? FoundPtr->Key : nullptr;
+}
+
+void UUVToolAABBTreeStorage::Remove(FDynamicMesh3* MeshKey)
+{
+	TreeInputObjectPair* Result = AABBTreeStorage.Find(MeshKey);
+	if (!Result)
+	{
+		return;
+	}
+	if (Result->Value.IsValid())
+	{
+		Result->Value->OnCanonicalModified.RemoveAll(this);
+	}
+
+	AABBTreeStorage.Remove(MeshKey);
+}
+
+void UUVToolAABBTreeStorage::RemoveByPredicate(TUniqueFunction<
+	bool(const FDynamicMesh3* Mesh, TWeakObjectPtr<UUVEditorToolMeshInput> InputObject,
+		TSharedPtr<FDynamicMeshAABBTree3> Tree)> Predicate)
+{
+	for (auto It = AABBTreeStorage.CreateIterator(); It; ++It)
+	{
+		TWeakObjectPtr<UUVEditorToolMeshInput> InputObject = It->Value.Value;
+		if (Predicate(It->Key, InputObject, It->Value.Key)) // mesh, object, tree
+		{
+			if (InputObject.IsValid())
+			{
+				InputObject->OnCanonicalModified.RemoveAll(this);
+			}
+			It.RemoveCurrent();
+		}
+	}
+}
+
+void UUVToolAABBTreeStorage::Empty()
+{
+	for (TPair<FDynamicMesh3*, TreeInputObjectPair>& MapPair : AABBTreeStorage)
+	{
+		TWeakObjectPtr<UUVEditorToolMeshInput> InputObject = MapPair.Value.Value;
+		if (InputObject.IsValid())
+		{
+			InputObject->OnCanonicalModified.RemoveAll(this);
+		}
+	}
+	AABBTreeStorage.Empty();
+}
+
+void UUVToolAABBTreeStorage::Shutdown()
+{
+	Empty();
+}
\ No newline at end of file
diff --git a/Engine/Plugins/Experimental/UVEditor/Source/UVEditorTools/Private/Selection/UVEditorMeshSelectionMechanic.cpp b/Engine/Plugins/Experimental/UVEditor/Source/UVEditorTools/Private/Selection/UVEditorMeshSelectionMechanic.cpp
--- a/Engine/Plugins/Experimental/UVEditor/Source/UVEditorTools/Private/Selection/UVEditorMeshSelectionMechanic.cpp
+++ b/Engine/Plugins/Experimental/UVEditor/Source/UVEditorTools/Private/Selection/UVEditorMeshSelectionMechanic.cpp
@@ -5,6 +5,9 @@
 #include "BaseBehaviors/BehaviorTargetInterfaces.h"
 #include "BaseBehaviors/SingleClickOrDragBehavior.h"
 #include "BaseBehaviors/MouseHoverBehavior.h"
+#include "ContextObjectStore.h"
+#include "ContextObjects/UVToolContextObjects.h"
+#include "ContextObjects/UVToolViewportButtonsAPI.h"
 #include "Drawing/TriangleSetComponent.h"
 #include "Drawing/LineSetComponent.h"
 #include "Drawing/PointSetComponent.h"
@@ -12,6 +15,7 @@
 #include "Intersection/IntrTriangle2AxisAlignedBox2.h"
 #include "Intersection/IntersectionQueries2.h"
 #include "InteractiveToolManager.h"
+#include "MeshOpPreviewHelpers.h" // UMeshOpPreviewWithBackgroundCompute
 #include "Polyline3.h"
 #include "Selections/MeshConnectedComponents.h"
 #include "Spatial/GeometrySet3.h"
@@ -25,48 +29,6 @@ using namespace UE::Geometry;
 
 namespace UVEditorMeshSelectionMechanicLocals
 {
-	/**
-	 * Class that can be used to apply/revert selection changes. Expects that the mesh selection
-	 * mechanic is the associated UObject that is passed to Apply/Revert.
-	 */
-	class FMeshSelectionMechanicSelectionChange : public FToolCommandChange
-	{
-	public:
-		FMeshSelectionMechanicSelectionChange(const FUVEditorDynamicMeshSelection& OldSelection,
-			const FUVEditorDynamicMeshSelection& NewSelection)
-			: Before(OldSelection)
-			, After(NewSelection)
-		{
-		}
-
-		virtual void Apply(UObject* Object) override
-		{
-			UUVEditorMeshSelectionMechanic* Mechanic = Cast<UUVEditorMeshSelectionMechanic>(Object);
-			if (Mechanic)
-			{
-				Mechanic->SetSelection(After, true, false);
-			}
-		}
-
-		virtual void Revert(UObject* Object) override
-		{
-			UUVEditorMeshSelectionMechanic* Mechanic = Cast<UUVEditorMeshSelectionMechanic>(Object);
-			if (Mechanic)
-			{
-				Mechanic->SetSelection(Before, true, false);
-			}
-		}
-
-		virtual FString ToString() const override
-		{
-			return TEXT("UVEditorMeshSelectionMechanicLocals::FMeshSelectionMechanicSelectionChange");
-		}
-
-	protected:
-		FUVEditorDynamicMeshSelection Before;
-		FUVEditorDynamicMeshSelection After;
-	};
-
 	template <typename InElementType>
 	void ToggleItem(TSet<InElementType>& Set, InElementType Item)
 	{
@@ -76,21 +38,26 @@ namespace UVEditorMeshSelectionMechanicLocals
 		}
 	}
 	
-	FUVEditorDynamicMeshSelection::EType ToCompatibleDynamicMeshSelectionType(const EUVEditorMeshSelectionMode& Mode)
+	using ESelectionMode = UUVToolSelectionAPI::EUVEditorSelectionMode;
+	using FModeChangeOptions = UUVToolSelectionAPI::FSelectionMechanicModeChangeOptions;
+
+	FUVToolSelection::EType ToCompatibleDynamicMeshSelectionType(ESelectionMode Mode)
 	{
 		switch (Mode)
 		{
-			case EUVEditorMeshSelectionMode::Mesh:
-			case EUVEditorMeshSelectionMode::Component:
-			case EUVEditorMeshSelectionMode::Triangle:
-				return FUVEditorDynamicMeshSelection::EType::Triangle;
-			case EUVEditorMeshSelectionMode::Edge:
-				return FUVEditorDynamicMeshSelection::EType::Edge;
-			case EUVEditorMeshSelectionMode::Vertex:
-				return FUVEditorDynamicMeshSelection::EType::Vertex;
+			case ESelectionMode::Mesh:
+			case ESelectionMode::Island:
+			case ESelectionMode::Triangle:
+				return FUVToolSelection::EType::Triangle;
+			case ESelectionMode::Edge:
+				return FUVToolSelection::EType::Edge;
+			case ESelectionMode::Vertex:
+				return FUVToolSelection::EType::Vertex;
+			case ESelectionMode::None: //doesn't actually matter what we return
+				return FUVToolSelection::EType::Vertex;
 		}
-		checkNoEntry();
-		return FUVEditorDynamicMeshSelection::EType::Vertex;
+		ensure(false);
+		return FUVToolSelection::EType::Vertex;
 	}
 
 	// Returns the marquee selection rectangle, obtained from the given CameraRectangle, projected to the XY plane
@@ -271,959 +238,793 @@ namespace UVEditorMeshSelectionMechanicLocals
 
 		return Result;
 	}
-} // namespace UVEditorMeshSelectionMechanicLocals
-
-void UUVEditorMeshSelectionMechanic::Setup(UInteractiveTool* ParentToolIn)
-{
-	UInteractionMechanic::Setup(ParentToolIn);
 
-	// This will be the target for the click drag behavior below
-	MarqueeMechanic = NewObject<URectangleMarqueeMechanic>();
-	MarqueeMechanic->bUseExternalClickDragBehavior = true;
-	MarqueeMechanic->Setup(ParentToolIn);
-	MarqueeMechanic->OnDragRectangleStarted.AddUObject(this, &UUVEditorMeshSelectionMechanic::OnDragRectangleStarted);
-	// TODO(Performance) :DynamicMarqueeSelection It would be cool to have the marquee selection update dynamically as
-	//  the rectangle gets changed, right now this isn't interactive for large meshes so we disabled it
-	// MarqueeMechanic->OnDragRectangleChanged.AddUObject(this, &UUVEditorMeshSelectionMechanic::OnDragRectangleChanged);
-	MarqueeMechanic->OnDragRectangleFinished.AddUObject(this, &UUVEditorMeshSelectionMechanic::OnDragRectangleFinished);
-
-	USingleClickOrDragInputBehavior* ClickOrDragBehavior = NewObject<USingleClickOrDragInputBehavior>();
-	ClickOrDragBehavior->Initialize(this, MarqueeMechanic);
-	ClickOrDragBehavior->Modifiers.RegisterModifier(ShiftModifierID, FInputDeviceState::IsShiftKeyDown);
-	ClickOrDragBehavior->Modifiers.RegisterModifier(CtrlModifierID, FInputDeviceState::IsCtrlKeyDown);
-	ParentTool->AddInputBehavior(ClickOrDragBehavior);
-	
-	Settings = NewObject<UUVEditorMeshSelectionMechanicProperties>(ParentToolIn);
-	Settings->RestoreProperties(ParentToolIn);
-	AddToolPropertySource(Settings);
-
-	ULocalMouseHoverBehavior* HoverBehavior = NewObject<ULocalMouseHoverBehavior>();
-	HoverBehavior->Initialize();
-	HoverBehavior->BeginHitTestFunc = [this](const FInputDeviceRay& Pos)
+	bool ConvertToHitElementList(ESelectionMode SelectionMode, 
+		const FDynamicMesh3& Mesh, int32 HitTid, const FViewCameraState& CameraState, 
+		const FRay& Ray, TArray<int32>& IDsOut)
 	{
-		FInputRayHit Hit;
-		if (Settings->bShowHoveredElements)
+		if (!ensure(HitTid != IndexConstants::InvalidID && Mesh.IsTriangle(HitTid)))
 		{
-			EUVEditorMeshSelectionMode Mode = SelectionMode;
-			if (Mode != EUVEditorMeshSelectionMode::Vertex && Mode != EUVEditorMeshSelectionMode::Edge)
-			{
-				Mode = EUVEditorMeshSelectionMode::Triangle;
-			}
-
-			// We don't bother with the depth since there aren't competing hover behaviors
-			Hit.bHit = RayCast(Pos, Mode).IsEmpty() == false;
+			return false;
 		}
-		else
+
+		IDsOut.Reset();
+
+		switch (SelectionMode)
 		{
-			// Disable hover behavior
-			Hit.bHit = false;
-		}
-		return Hit;
-	};
-	HoverBehavior->OnUpdateHoverFunc = [this](const FInputDeviceRay& Pos)
-	{
-		EUVEditorMeshSelectionMode Mode = SelectionMode;
-		if (SelectionMode != EUVEditorMeshSelectionMode::Vertex && SelectionMode != EUVEditorMeshSelectionMode::Edge)
+		case ESelectionMode::Island:
 		{
-			Mode = EUVEditorMeshSelectionMode::Triangle;
+			TRACE_CPUPROFILER_EVENT_SCOPE(Component);
+
+			FMeshConnectedComponents MeshSelectedComponent(&Mesh);
+			TArray<int32> SeedTriangles;
+			SeedTriangles.Add(HitTid);
+			MeshSelectedComponent.FindTrianglesConnectedToSeeds(SeedTriangles);
+			ensure(MeshSelectedComponent.Components.Num() == 1); // Expect each triangle to only be in a single component
+			IDsOut.Append(MoveTemp(MeshSelectedComponent.Components[0].Indices));
+			break;
 		}
-		
-		HoverPointSet->Clear();
-		HoverLineSet->Clear();
-		HoverTriangleSet->Clear();
-		
-		TSet<int32> Hovered = RayCast(Pos, Mode);
-		for (int32 Id : Hovered)
+		case ESelectionMode::Edge:
 		{
-			if (SelectionMode == EUVEditorMeshSelectionMode::Vertex)
+			TRACE_CPUPROFILER_EVENT_SCOPE(Edge);
+			// TODO: We'll need the ability to hit occluded triangles to see if there is a better edge to snap to.
+
+			// Try to snap to one of the edges.
+			FIndex3i Eids = Mesh.GetTriEdges(HitTid);
+
+			FGeometrySet3 GeometrySet;
+			for (int i = 0; i < 3; ++i)
 			{
-				const FVector& P = CurrentSelection.Mesh->GetVertexRef(Id);
-				const FRenderablePoint PointToRender(P, 
-					FUVEditorUXSettings::SelectionHoverTriangleWireframeColor, 
-					FUVEditorUXSettings::SelectionPointThickness);
-				
-				HoverPointSet->AddPoint(PointToRender);
+				FIndex2i Vids = Mesh.GetEdgeV(Eids[i]);
+				FPolyline3d Polyline(Mesh.GetVertex(Vids.A), Mesh.GetVertex(Vids.B));
+				GeometrySet.AddCurve(Eids[i], Polyline);
 			}
-			else if (SelectionMode == EUVEditorMeshSelectionMode::Edge)
+
+			FGeometrySet3::FNearest Nearest;
+			if (GeometrySet.FindNearestCurveToRay(Ray, Nearest,
+				[&CameraState](const FVector3d& Position1, const FVector3d& Position2) {
+					return ToolSceneQueriesUtil::PointSnapQuery(CameraState,
+						Position1, Position2,
+						ToolSceneQueriesUtil::GetDefaultVisualAngleSnapThreshD()); }))
 			{
-				const FIndex2i EdgeVids = CurrentSelection.Mesh->GetEdgeV(Id);
-				const FVector& A = CurrentSelection.Mesh->GetVertexRef(EdgeVids.A);
-				const FVector& B = CurrentSelection.Mesh->GetVertexRef(EdgeVids.B);
-				
-				HoverLineSet->AddLine(A, B, 
-					FUVEditorUXSettings::SelectionHoverTriangleWireframeColor, 
-					FUVEditorUXSettings::SelectionLineThickness,
-					FUVEditorUXSettings::SelectionHoverWireframeDepthBias);
+				IDsOut.Add(Nearest.ID);
 			}
-			else
+			break;
+		}
+		case ESelectionMode::Vertex:
+		{
+			TRACE_CPUPROFILER_EVENT_SCOPE(Vertex);
+			// TODO: Improve this to handle super narrow, sliver triangles better, where testing near vertices can be difficult.
+
+			// Try to snap to one of the vertices
+			FIndex3i Vids = Mesh.GetTriangle(HitTid);
+
+			FGeometrySet3 GeometrySet;
+			for (int i = 0; i < 3; ++i)
 			{
-				const FIndex3i Vids = CurrentSelection.Mesh->GetTriangle(Id);
-				const FVector& A = CurrentSelection.Mesh->GetVertex(Vids[0]);
-				const FVector& B = CurrentSelection.Mesh->GetVertex(Vids[1]);
-				const FVector& C = CurrentSelection.Mesh->GetVertex(Vids[2]);
-
-				HoverLineSet->AddLine(A, B, FUVEditorUXSettings::SelectionHoverTriangleWireframeColor, 
-					FUVEditorUXSettings::SelectionLineThickness, FUVEditorUXSettings::SelectionHoverWireframeDepthBias);
-				HoverLineSet->AddLine(B, C, FUVEditorUXSettings::SelectionHoverTriangleWireframeColor, 
-					FUVEditorUXSettings::SelectionLineThickness, FUVEditorUXSettings::SelectionHoverWireframeDepthBias);
-				HoverLineSet->AddLine(C, A, FUVEditorUXSettings::SelectionHoverTriangleWireframeColor, 
-					FUVEditorUXSettings::SelectionLineThickness, FUVEditorUXSettings::SelectionHoverWireframeDepthBias);
-				HoverTriangleSet->AddTriangle(A, B, C, FVector::ZAxisVector, 
-					FUVEditorUXSettings::SelectionHoverTriangleFillColor, HoverTriangleSetMaterial);
+				GeometrySet.AddPoint(Vids[i], Mesh.GetTriVertex(HitTid, i));
 			}
 
-			return true;
+			FGeometrySet3::FNearest Nearest;
+			if (GeometrySet.FindNearestPointToRay(Ray, Nearest,
+				[&CameraState](const FVector3d& Position1, const FVector3d& Position2) {
+					return ToolSceneQueriesUtil::PointSnapQuery(CameraState,
+						Position1, Position2,
+						ToolSceneQueriesUtil::GetDefaultVisualAngleSnapThreshD()); }))
+			{
+				IDsOut.Add(Nearest.ID);
+			}
+			break;
 		}
-		return false;
-	};
-	HoverBehavior->OnEndHoverFunc = [this]()
-	{
-		HoverPointSet->Clear();
-		HoverLineSet->Clear();
-		HoverTriangleSet->Clear();
-	};
-	ParentTool->AddInputBehavior(HoverBehavior);
+		case ESelectionMode::Triangle:
+		{
+			TRACE_CPUPROFILER_EVENT_SCOPE(Triangle);
 
-	ClearCurrentSelection();
-
-	// Setup selection visualizations
-	TriangleSet = NewObject<UTriangleSetComponent>();
-	// We are setting the TranslucencySortPriority here to handle the UV editor's use case in 2D
-	// where multiple translucent layers are drawn on top of each other but still need depth sorting.
-	TriangleSet->TranslucencySortPriority = FUVEditorUXSettings::SelectionTriangleDepthBias;	
-	TriangleSetMaterial = ToolSetupUtil::GetCustomTwoSidedDepthOffsetMaterial(GetParentTool()->GetToolManager(),
-		FUVEditorUXSettings::SelectionTriangleFillColor, 
-		FUVEditorUXSettings::SelectionTriangleDepthBias, 
-		FUVEditorUXSettings::SelectionTriangleOpacity);
-	
-	PointSet = NewObject<UPointSetComponent>();
-	PointSet->SetPointMaterial(ToolSetupUtil::GetDefaultPointComponentMaterial( GetParentTool()->GetToolManager(), true));
-	LineSet = NewObject<ULineSetComponent>();
-	LineSet->SetLineMaterial(ToolSetupUtil::GetDefaultLineComponentMaterial( GetParentTool()->GetToolManager(), true));
+			IDsOut.Add(HitTid);
+			break;
+		}
+		case ESelectionMode::Mesh:
+		{
+			TRACE_CPUPROFILER_EVENT_SCOPE(Mesh);
 
-	// Setup hover visualizations
-	HoverPointSet = NewObject<UPointSetComponent>();
-	HoverPointSet->SetPointMaterial(ToolSetupUtil::GetDefaultPointComponentMaterial(GetParentTool()->GetToolManager(), false));
-	HoverLineSet = NewObject<ULineSetComponent>();
-	HoverLineSet->SetLineMaterial(ToolSetupUtil::GetDefaultLineComponentMaterial(GetParentTool()->GetToolManager(), false));
-	HoverTriangleSet = NewObject<UTriangleSetComponent>();
-	HoverTriangleSetMaterial = ToolSetupUtil::GetCustomTwoSidedDepthOffsetMaterial(GetParentTool()->GetToolManager(),
-		FUVEditorUXSettings::SelectionHoverTriangleFillColor, 
-		FUVEditorUXSettings::SelectionHoverTriangleDepthBias, 
-		FUVEditorUXSettings::SelectionHoverTriangleOpacity);
+			for (int32 Tid : Mesh.TriangleIndicesItr())
+			{
+				IDsOut.Add(Tid);
+			}
+			break;
+		}
+		default:
+			ensure(false);
+			break;
+		}
 
-	// Add default selection change emitter if one was not provided.
-	if (!EmitSelectionChange)
-	{
-		EmitSelectionChange = [this](const FUVEditorDynamicMeshSelection& OldSelection, const FUVEditorDynamicMeshSelection& NewSelection)
-		{
-			GetParentTool()->GetToolManager()->EmitObjectChange(this, 
-				MakeUnique<UVEditorMeshSelectionMechanicLocals::FMeshSelectionMechanicSelectionChange>(OldSelection, NewSelection),
-				LOCTEXT("SelectionChangeMessage", "Selection Change"));
-		};
+		return !IDsOut.IsEmpty();
 	}
-}
 
-void UUVEditorMeshSelectionMechanic::SetWorld(UWorld* World)
-{
-	auto RegisterComponent = [this](UMeshComponent* Component)
+	/**
+	 * Undo/redo transaction for selection mode changes
+	 */
+	class  FModeChange : public FToolCommandChange
 	{
-		if (Component->IsRegistered())
+	public:
+		FModeChange(ESelectionMode BeforeIn, 
+			ESelectionMode AfterIn)
+			: Before(BeforeIn)
+			, After(AfterIn)
+		{};
+
+		virtual void Apply(UObject* Object) override
 		{
-			Component->ReregisterComponent();
+			UUVEditorMeshSelectionMechanic* SelectionMechanic = Cast<UUVEditorMeshSelectionMechanic>(Object);
+			FModeChangeOptions Options;
+			Options.bConvertExisting = false;
+			Options.bBroadcastIfConverted = false;
+			Options.bEmitChanges = false;
+			SelectionMechanic->SetSelectionMode(After, Options);
 		}
-		else
+
+		virtual void Revert(UObject* Object) override
+		{
+			UUVEditorMeshSelectionMechanic* SelectionMechanic = Cast<UUVEditorMeshSelectionMechanic>(Object);
+			FModeChangeOptions Options;
+			Options.bConvertExisting = false;
+			Options.bBroadcastIfConverted = false;
+			Options.bEmitChanges = false;
+			SelectionMechanic->SetSelectionMode(Before, Options);
+		}
+
+		virtual bool HasExpired(UObject* Object) const override
+		{
+			UUVEditorMeshSelectionMechanic* SelectionMechanic = Cast<UUVEditorMeshSelectionMechanic>(Object);
+			return !SelectionMechanic->IsEnabled();
+		}
+
+
+		virtual FString ToString() const override
 		{
-			Component->RegisterComponent();
+			return TEXT("UVEditorMeshSelectionMechanicLocals::FModeChange");
 		}
+
+	protected:
+		ESelectionMode Before;
+		ESelectionMode After;
 	};
-	
-	// It may be unreasonable to worry about SetWorld being called more than once, but let's be safe anyway
-	if (PreviewGeometryActor)
-	{
-		PreviewGeometryActor->Destroy();
-	}
+} // namespace UVEditorMeshSelectionMechanicLocals
 
-	// We need the world so we can create the geometry actor in the right place.
-	FRotator Rotation(0.0f, 0.0f, 0.0f);
-	FActorSpawnParameters SpawnInfo;
-	PreviewGeometryActor = World->SpawnActor<APreviewGeometryActor>(FVector::ZeroVector, Rotation, SpawnInfo);
 
+void UUVEditorMeshSelectionMechanic::Setup(UInteractiveTool* ParentToolIn)
+{
+	UInteractionMechanic::Setup(ParentToolIn);
 
-	// Attach the rendering component to the actor		
-	TriangleSet->Rename(nullptr, PreviewGeometryActor); // Changes the "outer"
-	PreviewGeometryActor->SetRootComponent(TriangleSet);
-	RegisterComponent(TriangleSet);
+	UContextObjectStore* ContextStore = GetParentTool()->GetToolManager()->GetContextObjectStore();
+	EmitChangeAPI = ContextStore->FindContext<UUVToolEmitChangeAPI>();
+	check(EmitChangeAPI);
 
-	LineSet->Rename(nullptr, PreviewGeometryActor); // Changes the "outer"
-	LineSet->AttachToComponent(TriangleSet, FAttachmentTransformRules::KeepWorldTransform);
-	RegisterComponent(LineSet);
-	
-	PointSet->Rename(nullptr, PreviewGeometryActor); // Changes the "outer"
-	PointSet->AttachToComponent(TriangleSet, FAttachmentTransformRules::KeepWorldTransform);
-	RegisterComponent(PointSet);
+	// This will be the target for the click drag behavior below
+	MarqueeMechanic = NewObject<URectangleMarqueeMechanic>();
+	MarqueeMechanic->bUseExternalClickDragBehavior = true;
+	MarqueeMechanic->Setup(ParentToolIn);
+	MarqueeMechanic->OnDragRectangleStarted.AddUObject(this, &UUVEditorMeshSelectionMechanic::OnDragRectangleStarted);
+	// TODO(Performance) :DynamicMarqueeSelection It would be cool to have the marquee selection update dynamically as
+	//  the rectangle gets changed, right now this isn't interactive for large meshes so we disabled it
+	//MarqueeMechanic->OnDragRectangleChanged.AddUObject(this, &UUVEditorMeshSelectionMechanic::OnDragRectangleChanged);
+	MarqueeMechanic->OnDragRectangleFinished.AddUObject(this, &UUVEditorMeshSelectionMechanic::OnDragRectangleFinished);
+
+	USingleClickOrDragInputBehavior* ClickOrDragBehavior = NewObject<USingleClickOrDragInputBehavior>();
+	ClickOrDragBehavior->Initialize(this, MarqueeMechanic);
+	ClickOrDragBehavior->Modifiers.RegisterModifier(ShiftModifierID, FInputDeviceState::IsShiftKeyDown);
+	ClickOrDragBehavior->Modifiers.RegisterModifier(CtrlModifierID, FInputDeviceState::IsCtrlKeyDown);
+	ParentTool->AddInputBehavior(ClickOrDragBehavior);
+
+	UMouseHoverBehavior* HoverBehavior = NewObject<UMouseHoverBehavior>();
+	HoverBehavior->Initialize(this);
+	ParentTool->AddInputBehavior(HoverBehavior);
 
+	ViewportButtonsAPI = ContextStore->FindContext<UUVToolViewportButtonsAPI>();
+	check(ViewportButtonsAPI);
+	ViewportButtonsAPI->OnSelectionModeChange.AddWeakLambda(this,
+		[this](UUVToolSelectionAPI::EUVEditorSelectionMode NewMode) {
+			SetSelectionMode(NewMode);
+		});
+	// Make sure we match the activated button
+	FModeChangeOptions ModeChangeOptions;
+	ModeChangeOptions.bEmitChanges = false;
+	SetSelectionMode(ViewportButtonsAPI->GetSelectionMode(), 
+		ModeChangeOptions); // convert, broadcast, don't emit
+
+	SetIsEnabled(bIsEnabled);
+}
+
+void UUVEditorMeshSelectionMechanic::Initialize(UWorld* World, UUVToolSelectionAPI* SelectionAPIIn)
+{
+	// It may be unreasonable to worry about Initialize being called more than once, but let's be safe anyway
 	if (HoverGeometryActor)
 	{
 		HoverGeometryActor->Destroy();
 	}
+
+	SelectionAPI = SelectionAPIIn;
 	
-	HoverGeometryActor = World->SpawnActor<APreviewGeometryActor>(FVector::ZeroVector, Rotation, SpawnInfo);
-	
-	// Attach the rendering component to the actor		
-	HoverTriangleSet->Rename(nullptr, HoverGeometryActor); // Changes the "outer"
-	HoverGeometryActor->SetRootComponent(HoverTriangleSet);
-	RegisterComponent(HoverTriangleSet);
-
-	HoverLineSet->Rename(nullptr, HoverGeometryActor); // Changes the "outer"
-	HoverLineSet->AttachToComponent(HoverTriangleSet, FAttachmentTransformRules::KeepWorldTransform);
-	RegisterComponent(HoverLineSet);
+	HoverGeometryActor = World->SpawnActor<APreviewGeometryActor>();
+
+	HoverTriangleSet = NewObject<UTriangleSetComponent>(HoverGeometryActor);
+	HoverTriangleSetMaterial = ToolSetupUtil::GetCustomTwoSidedDepthOffsetMaterial(GetParentTool()->GetToolManager(),
+		FUVEditorUXSettings::SelectionHoverTriangleFillColor,
+		FUVEditorUXSettings::SelectionHoverTriangleDepthBias,
+		FUVEditorUXSettings::SelectionHoverTriangleOpacity);
+	HoverGeometryActor->SetRootComponent(HoverTriangleSet.Get());
+	HoverTriangleSet->RegisterComponent();
+
+	HoverPointSet = NewObject<UPointSetComponent>(HoverGeometryActor);
+	HoverPointSet->SetPointMaterial(ToolSetupUtil::GetDefaultPointComponentMaterial(GetParentTool()->GetToolManager(), false));
+	HoverPointSet->AttachToComponent(HoverTriangleSet.Get(), FAttachmentTransformRules::KeepWorldTransform);
+	HoverPointSet->RegisterComponent();
 	
-	HoverPointSet->Rename(nullptr, HoverGeometryActor); // Changes the "outer"
-	HoverPointSet->AttachToComponent(HoverTriangleSet, FAttachmentTransformRules::KeepWorldTransform);
-	RegisterComponent(HoverPointSet);
+	HoverLineSet = NewObject<ULineSetComponent>(HoverGeometryActor);
+	HoverLineSet->SetLineMaterial(ToolSetupUtil::GetDefaultLineComponentMaterial(GetParentTool()->GetToolManager(), false));
+	HoverLineSet->AttachToComponent(HoverTriangleSet.Get(), FAttachmentTransformRules::KeepWorldTransform);
+	HoverLineSet->RegisterComponent();
 }
 
-void UUVEditorMeshSelectionMechanic::Shutdown()
+void UUVEditorMeshSelectionMechanic::SetIsEnabled(bool bIsEnabledIn)
 {
-	if (PreviewGeometryActor)
+	bIsEnabled = bIsEnabledIn;
+	if (MarqueeMechanic)
 	{
-		PreviewGeometryActor->Destroy();
-		PreviewGeometryActor = nullptr;
+		MarqueeMechanic->SetIsEnabled(bIsEnabled && SelectionMode != ESelectionMode::None);
+	}
+	if (ViewportButtonsAPI)
+	{
+		ViewportButtonsAPI->SetSelectionButtonsEnabled(bIsEnabledIn);
 	}
 }
 
-void UUVEditorMeshSelectionMechanic::AddSpatial(TSharedPtr<FDynamicMeshAABBTree3> SpatialIn, const FTransform& TransformIn)
+void UUVEditorMeshSelectionMechanic::SetShowHoveredElements(bool bShow)
 {
-	MeshSpatials.Add(SpatialIn);
-	MeshTransforms.Add(TransformIn);
+	bShowHoveredElements = bShow;
+	if (!bShowHoveredElements)
+	{
+		HoverPointSet->Clear();
+		HoverLineSet->Clear();
+		HoverTriangleSet->Clear();
+	}
 }
 
-const FUVEditorDynamicMeshSelection& UUVEditorMeshSelectionMechanic::GetCurrentSelection() const
+void UUVEditorMeshSelectionMechanic::Shutdown()
 {
-	return CurrentSelection;
+	if (HoverGeometryActor)
+	{
+		HoverGeometryActor->Destroy();
+		HoverGeometryActor = nullptr;
+	}
+	SelectionAPI = nullptr;
+	ViewportButtonsAPI = nullptr;
+	EmitChangeAPI = nullptr;
+	MarqueeMechanic = nullptr;
+	HoverTriangleSetMaterial = nullptr;
 }
 
-void UUVEditorMeshSelectionMechanic::SetSelection(const FUVEditorDynamicMeshSelection& Selection, bool bBroadcast, bool bEmitChange)
+void UUVEditorMeshSelectionMechanic::SetTargets(const TArray<TObjectPtr<UUVEditorToolMeshInput>>& TargetsIn)
 {
-	FUVEditorDynamicMeshSelection OriginalSelection = CurrentSelection;
-	CurrentSelection = Selection;
+	Targets = TargetsIn;
 
-	// Adjust the mesh in which the selection lives
-	if (CurrentSelection.Mesh == nullptr)
+	// Retrieve cached AABB tree storage, or else set it up
+	UContextObjectStore* ContextStore = ParentTool->GetToolManager()->GetContextObjectStore();
+	UUVToolAABBTreeStorage* TreeStore = ContextStore->FindContext<UUVToolAABBTreeStorage>();
+	if (!TreeStore)
 	{
-		// We're actually clearing the selection
-		if (!ensure(CurrentSelection.SelectedIDs.IsEmpty()))
-		{
-			CurrentSelection.SelectedIDs.Empty();
-		}
-		CurrentSelectionIndex = IndexConstants::InvalidID;
+		TreeStore = NewObject<UUVToolAABBTreeStorage>();
+		ContextStore->AddContextObject(TreeStore);
 	}
-	else if (CurrentSelectionIndex == IndexConstants::InvalidID 
-		|| MeshSpatials[CurrentSelectionIndex]->GetMesh() != Selection.Mesh)
+
+	// Get or create spatials
+	// Initialize the AABB trees from cached values, or make new ones
+	MeshSpatials.Reset();
+	for (TObjectPtr<UUVEditorToolMeshInput> Target : Targets)
 	{
-		CurrentSelectionIndex = IndexConstants::InvalidID;
-		for (int32 i = 0; i < MeshSpatials.Num(); ++i)
+		TSharedPtr<FDynamicMeshAABBTree3> Tree = TreeStore->Get(Target->UnwrapCanonical.Get());
+		if (!Tree)
 		{
-			if (MeshSpatials[i]->GetMesh() == Selection.Mesh)
-			{
-				CurrentSelectionIndex = i;
-				break;
-			}
+			TRACE_CPUPROFILER_EVENT_SCOPE(BuildAABBTreeForTarget);
+			Tree = MakeShared<FDynamicMeshAABBTree3>();
+			Tree->SetMesh(Target->UnwrapCanonical.Get(), false);
+			// For now we split round-robin on the X/Y axes TODO Experiment with better splitting heuristics
+			FDynamicMeshAABBTree3::GetSplitAxisFunc GetSplitAxis = [](int Depth, const FAxisAlignedBox3d&) { return Depth % 2; };
+			// Note: 16 tris/leaf was chosen with data collected by SpatialBenchmarks.cpp in GeometryProcessingUnitTests
+			Tree->SetBuildOptions(16, MoveTemp(GetSplitAxis));
+			Tree->Build();
+			TreeStore->Set(Target->UnwrapCanonical.Get(), Tree, Target);
 		}
-		ensure(CurrentSelectionIndex != IndexConstants::InvalidID);
+		MeshSpatials.Add(Tree);
 	}
+}
+
+void UUVEditorMeshSelectionMechanic::Render(IToolsContextRenderAPI* RenderAPI)
+{
+	MarqueeMechanic->Render(RenderAPI);
 
-	UpdateCentroid();
+	// Cache the camera state
+	GetParentTool()->GetToolManager()->GetContextQueriesAPI()->GetCurrentViewState(CameraState);
+}
 
-	if (bEmitChange && OriginalSelection != CurrentSelection)
-	{
-		EmitSelectionChange(OriginalSelection, CurrentSelection);
-	}
-	if (bBroadcast)
-	{
-		OnSelectionChanged.Broadcast();
-	}
+void UUVEditorMeshSelectionMechanic::DrawHUD(FCanvas* Canvas, IToolsContextRenderAPI* RenderAPI)
+{
+	MarqueeMechanic->DrawHUD(Canvas, RenderAPI);
+}
 
-	// Rebuild after broadcast in case the outside world wants to adjust things like color...
-	RebuildDrawnElements(FTransform(CurrentSelectionCentroid));
+FInputRayHit UUVEditorMeshSelectionMechanic::IsHitByClick(const FInputDeviceRay& ClickPos)
+{
+	FInputRayHit Hit; 
+	// If enabled, return a hit so we always capture and can clear the selection
+	Hit.bHit = bIsEnabled && SelectionMode != ESelectionMode::None;
+	return Hit;	
 }
 
-void UUVEditorMeshSelectionMechanic::RebuildDrawnElements(const FTransform& StartTransform)
+
+void UUVEditorMeshSelectionMechanic::SetSelectionMode(
+	ESelectionMode TargetMode, const FModeChangeOptions& Options)
 {
-	TRACE_CPUPROFILER_EVENT_SCOPE(MeshSelectionMechanic_RebuildDrawnElements);
+	TRACE_CPUPROFILER_EVENT_SCOPE(MeshSelectionMechanic_ChangeSelectionMode);
 
-	TriangleSet->Clear();
-	LineSet->Clear();
-	PointSet->Clear();
-	PreviewGeometryActor->SetActorTransform(StartTransform);
+	const FText TransactionName = LOCTEXT("ChangeSelectionMode", "Change Selection Mode");
+	
+	using namespace UVEditorMeshSelectionMechanicLocals;
 
-	if (CurrentSelectionIndex == IndexConstants::InvalidID)
+	ESelectionMode OldMode = SelectionMode;
+	SelectionMode = TargetMode;
+	if (OldMode == SelectionMode || !bIsEnabled || !SelectionAPI)
 	{
 		return;
 	}
 
-	// For us to end up with the StartTransform, we have to invert it, but only after applying
-	// the mesh transform to begin with.
-	auto TransformToApply = [&StartTransform, this](const FVector& VectorIn)
-	{
-		return StartTransform.InverseTransformPosition(
-			MeshTransforms[CurrentSelectionIndex].TransformPosition(VectorIn));
-	};
-
-	if (CurrentSelection.Type == FUVEditorDynamicMeshSelection::EType::Triangle)
+	if (ViewportButtonsAPI)
 	{
-		TRACE_CPUPROFILER_EVENT_SCOPE(MeshSelectionMechanic_RebuildDrawnElements_Triangle);
-		
-		TriangleSet->ReserveTriangles(CurrentSelection.SelectedIDs.Num());
-		LineSet->ReserveLines(CurrentSelection.SelectedIDs.Num() * 3);
-		for (int32 Tid : CurrentSelection.SelectedIDs)
-		{
-			FIndex3i Vids = CurrentSelection.Mesh->GetTriangle(Tid);
-			FVector Points[3];
-			for (int i = 0; i < 3; ++i)
-			{
-				Points[i] = TransformToApply(CurrentSelection.Mesh->GetVertex(Vids[i]));
-			}
-			TriangleSet->AddTriangle(Points[0], Points[1], Points[2], FVector(0, 0, 1), FUVEditorUXSettings::SelectionTriangleFillColor, TriangleSetMaterial);
-			for (int i = 0; i < 3; ++i)
-			{
-				int NextIndex = (i + 1) % 3;
-				LineSet->AddLine(Points[i], Points[NextIndex],
-					FUVEditorUXSettings::SelectionTriangleWireframeColor, 
-					FUVEditorUXSettings::SelectionLineThickness, 
-					FUVEditorUXSettings::SelectionWireframeDepthBias);
-			}
-		}
+		// Not clear whether we should or shouldn't broadcast this. A user could conceivably set selection
+		// via mechanic and expect for a notification from the viewport buttons, but it feels wrong to
+		// knowingly trigger a second call into this function if we broadcast, and that example seems like
+		// questionable code organization...
+		ViewportButtonsAPI->SetSelectionMode(SelectionMode, false);
 	}
-	else if (CurrentSelection.Type == FUVEditorDynamicMeshSelection::EType::Edge)
-	{
-		TRACE_CPUPROFILER_EVENT_SCOPE(MeshSelectionMechanic_RebuildDrawnElements_Edge);
 
-		LineSet->ReserveLines(CurrentSelection.SelectedIDs.Num());
-		for (int32 Eid : CurrentSelection.SelectedIDs)
-		{
-			FIndex2i EdgeVids = CurrentSelection.Mesh->GetEdgeV(Eid);
-			LineSet->AddLine(
-				TransformToApply(CurrentSelection.Mesh->GetVertex(EdgeVids.A)),
-				TransformToApply(CurrentSelection.Mesh->GetVertex(EdgeVids.B)),
-				FUVEditorUXSettings::SelectionTriangleWireframeColor, 
-				FUVEditorUXSettings::SelectionLineThickness, 
-				FUVEditorUXSettings::SelectionWireframeDepthBias);
-		}
-	}
-	else if (CurrentSelection.Type == FUVEditorDynamicMeshSelection::EType::Vertex)
+	if (Options.bEmitChanges)
 	{
-		TRACE_CPUPROFILER_EVENT_SCOPE(MeshSelectionMechanic_RebuildDrawnElements_Vertex);
-
-		PointSet->ReservePoints(CurrentSelection.SelectedIDs.Num());
-		for (int32 Vid : CurrentSelection.SelectedIDs)
-		{
-			FRenderablePoint PointToRender(TransformToApply(CurrentSelection.Mesh->GetVertex(Vid)),
-				FUVEditorUXSettings::SelectionTriangleWireframeColor,
-				FUVEditorUXSettings::SelectionPointThickness,
-				FUVEditorUXSettings::SelectionWireframeDepthBias);
-			PointSet->AddPoint(PointToRender);
-		}
+		EmitChangeAPI->BeginUndoTransaction(TransactionName);
+		EmitChangeAPI->EmitToolIndependentChange(this, MakeUnique<FModeChange>(OldMode, SelectionMode), TransactionName);
 	}
-}
 
-void UUVEditorMeshSelectionMechanic::UpdateCentroid()
-{
-	TRACE_CPUPROFILER_EVENT_SCOPE(MeshSelectionMechanic_UpdateCentroid);
+	MarqueeMechanic->SetIsEnabled(bIsEnabled && SelectionMode != ESelectionMode::None);
 
-	CurrentSelectionCentroid = FVector3d(0);
-	if (CurrentSelection.IsEmpty())
+	// See whether a conversion is not necessary
+	FUVToolSelection::EType ExpectedSelectionType = ToCompatibleDynamicMeshSelectionType(SelectionMode);
+	FUVToolSelection::EType CurrentSelectionType = SelectionAPI->GetSelectionsType();
+	if (!SelectionAPI->HaveSelections() || ExpectedSelectionType == CurrentSelectionType 
+		|| !Options.bConvertExisting || SelectionMode == ESelectionMode::None)
 	{
+		// No conversion needed
+		if (Options.bEmitChanges)
+		{
+			EmitChangeAPI->EndUndoTransaction();
+		}
 		return;
 	}
+	
+	// We're going to convert the existing selection.
+	const TArray<FUVToolSelection>& OriginalSelections = SelectionAPI->GetSelections();
+	TArray<FUVToolSelection> NewSelections;
 
-	if (CurrentSelection.Type == FUVEditorDynamicMeshSelection::EType::Edge)
+	for (const FUVToolSelection& OriginalSelection : OriginalSelections)
 	{
-		for (int32 Eid : CurrentSelection.SelectedIDs)
-		{
-			CurrentSelectionCentroid += CurrentSelection.Mesh->GetEdgePoint(Eid, 0.5);
-		}
-		CurrentSelectionCentroid /= CurrentSelection.SelectedIDs.Num();
-	}
-	else if (CurrentSelection.Type == FUVEditorDynamicMeshSelection::EType::Triangle)
-	{
-		for (int32 Tid : CurrentSelection.SelectedIDs)
-		{
-			CurrentSelectionCentroid += CurrentSelection.Mesh->GetTriCentroid(Tid);
-		}
-		CurrentSelectionCentroid /= CurrentSelection.SelectedIDs.Num();
-	}
-	if (CurrentSelection.Type == FUVEditorDynamicMeshSelection::EType::Vertex)
-	{
-		for (int32 Vid : CurrentSelection.SelectedIDs)
-		{
-			CurrentSelectionCentroid += CurrentSelection.Mesh->GetVertex(Vid);
-		}
-		CurrentSelectionCentroid /= CurrentSelection.SelectedIDs.Num();
-	}
-
-	// disable centroid caching if mesh does not have shape changestamp enabled
-	CentroidTimestamp = CurrentSelection.Mesh->HasShapeChangeStampEnabled() ?
-		CurrentSelection.Mesh->GetShapeChangeStamp() : TNumericLimits<uint32>::Max();
-}
-
-FVector3d UUVEditorMeshSelectionMechanic::GetCurrentSelectionCentroid()
-{
-	if (!CurrentSelection.Mesh)
-	{
-		return FVector3d(0);
-	}
-	else if (CurrentSelection.Mesh->GetShapeChangeStamp() != CentroidTimestamp)
-	{
-		UpdateCentroid();
-	}
-
-	return CurrentSelectionCentroid;
-}
-
-
-void UUVEditorMeshSelectionMechanic::SetDrawnElementsTransform(const FTransform& Transform)
-{
-	PreviewGeometryActor->SetActorTransform(Transform);
-}
-
-void UUVEditorMeshSelectionMechanic::Render(IToolsContextRenderAPI* RenderAPI)
-{
-	// Cache the camera state
-	MarqueeMechanic->Render(RenderAPI);
-	GetParentTool()->GetToolManager()->GetContextQueriesAPI()->GetCurrentViewState(CameraState);
-}
-
-void UUVEditorMeshSelectionMechanic::DrawHUD(FCanvas* Canvas, IToolsContextRenderAPI* RenderAPI)
-{
-	MarqueeMechanic->DrawHUD(Canvas, RenderAPI);
-}
-
-FInputRayHit UUVEditorMeshSelectionMechanic::IsHitByClick(const FInputDeviceRay& ClickPos)
-{
-	// TODO We could account for what modifiers would do and return an accurate depth here but for now this simple code works fine
-
-	// Return a hit so we always capture and can clear the selection
-	FInputRayHit Hit; 
-	Hit.bHit = true;
-	return Hit;	
-}
-
-void UUVEditorMeshSelectionMechanic::ClearCurrentSelection()
-{
-	CurrentSelection.SelectedIDs.Empty();
-	CurrentSelection.Mesh = nullptr;
-	CurrentSelection.Type = UVEditorMeshSelectionMechanicLocals::ToCompatibleDynamicMeshSelectionType(SelectionMode);
-	CurrentSelectionIndex = IndexConstants::InvalidID;
-}
-
-
-void UUVEditorMeshSelectionMechanic::ChangeSelectionMode(const EUVEditorMeshSelectionMode& TargetMode)
-{
-	TRACE_CPUPROFILER_EVENT_SCOPE(MeshSelectionMechanic_ChangeSelectionMode);
-	
-	if (CurrentSelection.SelectedIDs.IsEmpty())
-	{
-		SelectionMode = TargetMode;
-		return;
-	}
-	
-	const FUVEditorDynamicMeshSelection OriginalSelection = CurrentSelection;
-	
-	auto VerticesToEdges = [this, &OriginalSelection]()
-	{
-		TRACE_CPUPROFILER_EVENT_SCOPE(MeshSelectionMechanic_ChangeSelectionMode_VerticesToEdges);
+		FDynamicMesh3* Mesh = OriginalSelection.Target->UnwrapCanonical.Get();
 		
-		CurrentSelection.SelectedIDs.Empty();
-		CurrentSelection.Type = FUVEditorDynamicMeshSelection::EType::Edge;
+		FUVToolSelection NewSelection;
+		NewSelection.Target = OriginalSelection.Target;
+		NewSelection.Type = ExpectedSelectionType;
+
+		const TSet<int32> OriginalIDs = OriginalSelection.SelectedIDs;
+		TSet<int32>& NewIDs = NewSelection.SelectedIDs;
 
-		for (int32 Vid : OriginalSelection.SelectedIDs)
+		auto VerticesToEdges = [Mesh, &OriginalIDs, &NewIDs]()
 		{
-			for (int32 Eid : OriginalSelection.Mesh->VtxEdgesItr(Vid))
+			TRACE_CPUPROFILER_EVENT_SCOPE(MeshSelectionMechanic_ChangeSelectionMode_VerticesToEdges);
+
+			for (int32 Vid : OriginalIDs)
 			{
-				if (!CurrentSelection.SelectedIDs.Contains(Eid))
+				for (int32 Eid : Mesh->VtxEdgesItr(Vid))
 				{
-					FIndex2i Verts = OriginalSelection.Mesh->GetEdgeV(Eid);
-					if (OriginalSelection.SelectedIDs.Contains(Verts.A) &&
-						OriginalSelection.SelectedIDs.Contains(Verts.B))
+					if (!NewIDs.Contains(Eid))
 					{
-						CurrentSelection.SelectedIDs.Add(Eid);	
+						FIndex2i Verts = Mesh->GetEdgeV(Eid);
+						if (OriginalIDs.Contains(Verts.A) &&
+							OriginalIDs.Contains(Verts.B))
+						{
+							NewIDs.Add(Eid);
+						}
 					}
 				}
 			}
-		}
-	};
+		};
 
-	auto VerticesToTriangles = [this, &OriginalSelection]()
-	{
-		TRACE_CPUPROFILER_EVENT_SCOPE(MeshSelectionMechanic_ChangeSelectionMode_VerticesToTriangles);
-		
-		CurrentSelection.SelectedIDs.Empty();
-		CurrentSelection.Type = FUVEditorDynamicMeshSelection::EType::Triangle;
-		
-		for (int32 Vid : OriginalSelection.SelectedIDs)
+		auto VerticesToTriangles = [Mesh, &OriginalIDs, &NewIDs]()
 		{
-			for (int32 Tid : OriginalSelection.Mesh->VtxTrianglesItr(Vid))
+			TRACE_CPUPROFILER_EVENT_SCOPE(MeshSelectionMechanic_ChangeSelectionMode_VerticesToTriangles);
+
+			for (int32 Vid : OriginalIDs)
 			{
-				if (!CurrentSelection.SelectedIDs.Contains(Tid))
+				for (int32 Tid : Mesh->VtxTrianglesItr(Vid))
 				{
-					FIndex3i Verts = OriginalSelection.Mesh->GetTriangle(Tid);
-					if (OriginalSelection.SelectedIDs.Contains(Verts.A) &&
-						OriginalSelection.SelectedIDs.Contains(Verts.B) &&
-						OriginalSelection.SelectedIDs.Contains(Verts.C))
+					if (!NewIDs.Contains(Tid))
 					{
-						CurrentSelection.SelectedIDs.Add(Tid);	
+						FIndex3i Verts = Mesh->GetTriangle(Tid);
+						if (OriginalIDs.Contains(Verts.A) &&
+							OriginalIDs.Contains(Verts.B) &&
+							OriginalIDs.Contains(Verts.C))
+						{
+							NewIDs.Add(Tid);
+						}
 					}
 				}
 			}
-		}
-	};
+		};
 
-	auto EdgesToVertices = [this, &OriginalSelection]()
-	{
-		TRACE_CPUPROFILER_EVENT_SCOPE(MeshSelectionMechanic_ChangeSelectionMode_EdgesToVertices);
-		
-		CurrentSelection.SelectedIDs.Empty();
-		CurrentSelection.Type = FUVEditorDynamicMeshSelection::EType::Vertex;
-		
-		for (int32 Eid : OriginalSelection.SelectedIDs)
+		auto EdgesToVertices = [Mesh, &OriginalIDs, &NewIDs]()
 		{
-			FIndex2i Verts = OriginalSelection.Mesh->GetEdgeV(Eid);
-			CurrentSelection.SelectedIDs.Add(Verts.A);
-			CurrentSelection.SelectedIDs.Add(Verts.B);
-		}
-	};
+			TRACE_CPUPROFILER_EVENT_SCOPE(MeshSelectionMechanic_ChangeSelectionMode_EdgesToVertices);
 
-	// Triangles with two selected edges will be selected
-	auto EdgesToTriangles = [this, &OriginalSelection]()
-	{
-		TRACE_CPUPROFILER_EVENT_SCOPE(MeshSelectionMechanic_ChangeSelectionMode_EdgesToTriangles);
-		
-		CurrentSelection.SelectedIDs.Empty();
-		CurrentSelection.Type = FUVEditorDynamicMeshSelection::EType::Triangle;
-		
-		TArray<int32> FoundTriangles;
-		for (int32 Eid : OriginalSelection.SelectedIDs)
-		{
-			FIndex2i Tris = OriginalSelection.Mesh->GetEdgeT(Eid);
-			FoundTriangles.Add(Tris.A);
-			if (Tris.B != IndexConstants::InvalidID)
+			for (int32 Eid : OriginalIDs)
 			{
-				FoundTriangles.Add(Tris.B);
+				FIndex2i Verts = Mesh->GetEdgeV(Eid);
+				NewIDs.Add(Verts.A);
+				NewIDs.Add(Verts.B);
 			}
-		}
+		};
 
-		if (FoundTriangles.Num() < 2)
+		// Triangles with two selected edges will be selected
+		auto EdgesToTriangles = [Mesh, &OriginalIDs, &NewIDs]()
 		{
-			return;
-		}
+			TRACE_CPUPROFILER_EVENT_SCOPE(MeshSelectionMechanic_ChangeSelectionMode_EdgesToTriangles);
 
-		Algo::Sort(FoundTriangles);
+			TArray<int32> FoundTriangles;
+			for (int32 Eid : OriginalIDs)
+			{
+				FIndex2i Tris = Mesh->GetEdgeT(Eid);
+				FoundTriangles.Add(Tris.A);
+				if (Tris.B != IndexConstants::InvalidID)
+				{
+					FoundTriangles.Add(Tris.B);
+				}
+			}
 
-		for (int I = 0; I < FoundTriangles.Num() - 1; I++)
-		{
-			if (FoundTriangles[I] == FoundTriangles[I + 1])
+			if (FoundTriangles.Num() < 2)
 			{
-				CurrentSelection.SelectedIDs.Add(FoundTriangles[I]);
-				I++;
-			}	
-		}
-	};
+				return;
+			}
 
-	auto TrianglesToVertices = [this, &OriginalSelection]()
-	{
-		TRACE_CPUPROFILER_EVENT_SCOPE(MeshSelectionMechanic_ChangeSelectionMode_TrianglesToVertices);
-		
-		CurrentSelection.SelectedIDs.Empty();
-		CurrentSelection.Type = FUVEditorDynamicMeshSelection::EType::Vertex;
+			Algo::Sort(FoundTriangles);
 
-		for (int32 Tid : OriginalSelection.SelectedIDs)
-		{
-			FIndex3i Verts = OriginalSelection.Mesh->GetTriangle(Tid);
-			CurrentSelection.SelectedIDs.Add(Verts.A);
-			CurrentSelection.SelectedIDs.Add(Verts.B);
-			CurrentSelection.SelectedIDs.Add(Verts.C);
-		}
-	};
+			for (int I = 0; I < FoundTriangles.Num() - 1; I++)
+			{
+				if (FoundTriangles[I] == FoundTriangles[I + 1])
+				{
+					NewIDs.Add(FoundTriangles[I]);
+					I++;
+				}
+			}
+		};
 
-	auto TrianglesToEdges = [this, &OriginalSelection]()
-	{
-		TRACE_CPUPROFILER_EVENT_SCOPE(MeshSelectionMechanic_ChangeSelectionMode_TrianglesToEdges);
-			
-		CurrentSelection.SelectedIDs.Empty();
-		CurrentSelection.Type = FUVEditorDynamicMeshSelection::EType::Edge;
-	
-		for (int32 Tid : OriginalSelection.SelectedIDs)
+		auto TrianglesToVertices = [Mesh, &OriginalIDs, &NewIDs]()
 		{
-			FIndex3i Edges = OriginalSelection.Mesh->GetTriEdgesRef(Tid);
-			CurrentSelection.SelectedIDs.Add(Edges.A);
-			CurrentSelection.SelectedIDs.Add(Edges.B);
-			CurrentSelection.SelectedIDs.Add(Edges.C);
-		}
-	};
+			TRACE_CPUPROFILER_EVENT_SCOPE(MeshSelectionMechanic_ChangeSelectionMode_TrianglesToVertices);
 
-	switch (SelectionMode)
-	{
-	case EUVEditorMeshSelectionMode::Vertex:
-		switch (TargetMode)
-		{
-		case EUVEditorMeshSelectionMode::Vertex:
-			// Do nothing
-			break;
+			for (int32 Tid : OriginalIDs)
+			{
+				FIndex3i Verts = Mesh->GetTriangle(Tid);
+				NewIDs.Add(Verts.A);
+				NewIDs.Add(Verts.B);
+				NewIDs.Add(Verts.C);
+			}
+		};
 
-		case EUVEditorMeshSelectionMode::Edge:
-			VerticesToEdges();
-			break;
-			
-		case EUVEditorMeshSelectionMode::Triangle:
-		case EUVEditorMeshSelectionMode::Component:
-		case EUVEditorMeshSelectionMode::Mesh:
-			VerticesToTriangles();
-			break;
-		}
-		break; // SelectionMode
-	
-	case EUVEditorMeshSelectionMode::Edge:
-		switch (TargetMode)
+		auto TrianglesToEdges = [Mesh, &OriginalIDs, &NewIDs]()
 		{
-		case EUVEditorMeshSelectionMode::Vertex:
-			EdgesToVertices();
-			break;
+			TRACE_CPUPROFILER_EVENT_SCOPE(MeshSelectionMechanic_ChangeSelectionMode_TrianglesToEdges);
 
-		case EUVEditorMeshSelectionMode::Edge:
-			// Do nothing
-			break;
+			for (int32 Tid : OriginalIDs)
+			{
+				FIndex3i Edges = Mesh->GetTriEdgesRef(Tid);
+				NewIDs.Add(Edges.A);
+				NewIDs.Add(Edges.B);
+				NewIDs.Add(Edges.C);
+			}
+		};
 
-		case EUVEditorMeshSelectionMode::Mesh:
-		case EUVEditorMeshSelectionMode::Triangle:
-		case EUVEditorMeshSelectionMode::Component:
-			EdgesToTriangles();
-			break;
-		}
-		break; // SelectionMode
-	
-	case EUVEditorMeshSelectionMode::Mesh:
-	case EUVEditorMeshSelectionMode::Triangle:
-	case EUVEditorMeshSelectionMode::Component:
-		switch (TargetMode)
+		switch (CurrentSelectionType)
 		{
-		case EUVEditorMeshSelectionMode::Vertex:
-			TrianglesToVertices();
-			break;
-
-		case EUVEditorMeshSelectionMode::Edge:
-			TrianglesToEdges();
-			break;
+		case FUVToolSelection::EType::Vertex:
+			switch (ExpectedSelectionType)
+			{
+			case FUVToolSelection::EType::Vertex:
+				ensure(false); // Should have been an early-out
+				break;
+			case FUVToolSelection::EType::Edge:
+				VerticesToEdges();
+				break;
+			case FUVToolSelection::EType::Triangle:
+				VerticesToTriangles();
+				break;
+			}
+			break; // SelectionMode
+		case FUVToolSelection::EType::Edge:
+			switch (ExpectedSelectionType)
+			{
+			case FUVToolSelection::EType::Vertex:
+				EdgesToVertices();
+				break;
+			case FUVToolSelection::EType::Edge:
+				ensure(false); // Should have been an early-out
+				break;
+			case FUVToolSelection::EType::Triangle:
+				EdgesToTriangles();
+				break;
+			}
+			break; // SelectionMode
+		case FUVToolSelection::EType::Triangle:
+			switch (ExpectedSelectionType)
+			{
+			case FUVToolSelection::EType::Vertex:
+				TrianglesToVertices();
+				break;
+			case FUVToolSelection::EType::Edge:
+				TrianglesToEdges();
+				break;
+			case FUVToolSelection::EType::Triangle:
+				ensure(false); // Should have been an early-out
+				break;
+			}
+			break; // SelectionMode
+		}
 
-		case EUVEditorMeshSelectionMode::Mesh:
-		case EUVEditorMeshSelectionMode::Triangle:
-		case EUVEditorMeshSelectionMode::Component:
-			// Do nothing
-			break;
+		if (!NewSelection.IsEmpty())
+		{
+			NewSelections.Add(MoveTemp(NewSelection));
 		}
-		break; // SelectionMode
 	}
 
-	if (CurrentSelection != OriginalSelection)
+	// Apply selection change
+	SelectionAPI->SetSelections(NewSelections, Options.bBroadcastIfConverted, 
+		Options.bEmitChanges);
+
+	if (Options.bEmitChanges)
 	{
-		UpdateCentroid();
-		RebuildDrawnElements(FTransform(GetCurrentSelectionCentroid()));
-		EmitSelectionChange(OriginalSelection, CurrentSelection);
-		OnSelectionChanged.Broadcast();
+		EmitChangeAPI->EndUndoTransaction();
 	}
-
-	SelectionMode = TargetMode;
+	
 	return;
 }
 
-
-void UUVEditorMeshSelectionMechanic::UpdateCurrentSelection(const TSet<int32>& NewSelection, bool CalledFromOnDragRectangleChanged)
+void UUVEditorMeshSelectionMechanic::ModifyExistingSelection(TSet<int32>& SelectionSetToModify, 
+	const TArray<int32>& SelectedIDs)
 {
 	TRACE_CPUPROFILER_EVENT_SCOPE(MeshSelectionMechanic_UpdateCurrentSelection);
 	
-	if (!ShouldRestartSelection() && CalledFromOnDragRectangleChanged)
-	{
-		// If we're modifying (adding/removing/toggling, but not restarting) the selection, we should start with the
-		// selection cached in OnDragRectangleStarted, otherwise multiple changes get accumulated as the rectangle is
-		// swept around
-		CurrentSelection.SelectedIDs = PreDragSelection.SelectedIDs;
-	}
-	
-	if (ShouldRestartSelection())
-	{
-		CurrentSelection.SelectedIDs = NewSelection;
-	}
-	else if (ShouldAddToSelection())
+	using namespace UVEditorMeshSelectionMechanicLocals;
+
+	if (ShouldAddToSelection())
 	{
-		CurrentSelection.SelectedIDs.Append(NewSelection);
+		SelectionSetToModify.Append(SelectedIDs);
 	}
 	else if (ShouldToggleFromSelection())
 	{
-		for (int32 Index : NewSelection)
+		for (int32 ID : SelectedIDs)
 		{
-			UVEditorMeshSelectionMechanicLocals::ToggleItem(CurrentSelection.SelectedIDs, Index);
+			ToggleItem(SelectionSetToModify, ID);
 		}
 	}
 	else if (ShouldRemoveFromSelection())
 	{
-		CurrentSelection.SelectedIDs = CurrentSelection.SelectedIDs.Difference(NewSelection);
+		SelectionSetToModify = SelectionSetToModify.Difference(TSet<int32>(SelectedIDs));
 	}
 	else
 	{
-		checkNoEntry();
-	}
-
-	if (CurrentSelection.SelectedIDs.IsEmpty())
-	{
-		ClearCurrentSelection();
+		// We shouldn't be trying to modify an existing selection if we're supposed to restart
+		ensure(false);
 	}
-	
-	CurrentSelection.Type = UVEditorMeshSelectionMechanicLocals::ToCompatibleDynamicMeshSelectionType(SelectionMode);
 }
 
 void UUVEditorMeshSelectionMechanic::OnClicked(const FInputDeviceRay& ClickPos)
 {
 	TRACE_CPUPROFILER_EVENT_SCOPE(MeshSelectionMechanic_OnClicked);
 	
-	FUVEditorDynamicMeshSelection OriginalSelection = CurrentSelection;
-	
-	if (CurrentSelection.Type != UVEditorMeshSelectionMechanicLocals::ToCompatibleDynamicMeshSelectionType(SelectionMode))
+	using namespace UVEditorMeshSelectionMechanicLocals;
+
+	// IsHitByClick should prevent us being here with !bIsEnabled
+	if (!ensure(bIsEnabled))
 	{
-		ClearCurrentSelection();
+		return;
 	}
 
-	TSet<int32> ClickSelectedIDs = RayCast(ClickPos, SelectionMode); // TODO Maybe using a TVariant would have been better here...
+	FUVToolSelection::EType ElementType = ToCompatibleDynamicMeshSelectionType(SelectionMode);
 
-	// TODO Perhaps selection clearing should happen only if the click occurs further than some threshold from the meshes
-	UpdateCurrentSelection(ClickSelectedIDs);
+	int32 HitAssetID = IndexConstants::InvalidID;
+	int32 HitTid = IndexConstants::InvalidID;
+	int32 ExistingSelectionIndex = IndexConstants::InvalidID;
+	TArray<int32> NewIDs;
 
-	if (OriginalSelection != CurrentSelection)
+	// Do the raycast and get selected elements
+	if (!GetHitTid(ClickPos, HitTid, HitAssetID, &ExistingSelectionIndex)
+		|| !ConvertToHitElementList(SelectionMode, *Targets[HitAssetID]->UnwrapCanonical,
+				HitTid, CameraState, ClickPos.WorldRay, NewIDs))
 	{
-		UpdateCentroid();
-		EmitSelectionChange(OriginalSelection, CurrentSelection);
-		OnSelectionChanged.Broadcast();
+		// Failed to select an element. See if selection needs clearing, and exit.
+		if (ShouldRestartSelection() && SelectionAPI->HaveSelections())
+		{
+			SelectionAPI->ClearSelections(true, true); // broadcast and emit
+		}
+		return;
+	}
 
-		// Rebuild after broadcast in case the outside world wants to adjust things like color...
-		RebuildDrawnElements(FTransform(GetCurrentSelectionCentroid()));
+	TArray<FUVToolSelection> NewSelections;
+	if (!ShouldRestartSelection())
+	{
+		NewSelections = SelectionAPI->GetSelections();
 	}
-}
 
-TSet<int32> UUVEditorMeshSelectionMechanic::RayCast(const FInputDeviceRay& Pos, EUVEditorMeshSelectionMode Mode)
-{
-	// TODO if the selection mode is Vertex, Edge or Triangle we'll return at most one index, if the selection mode is
-	// Component or Mesh we'll return a set of components. Perhaps it would be better to use TVariant here (and maybe
-	// also a TArray rather than a TSet)
-	TSet<int32> Result;
-	
-	int32 HitTid = IndexConstants::InvalidID;
-	for (int32 MeshIndex = 0; MeshIndex < MeshSpatials.Num(); ++MeshIndex)
+	if (NewIDs.IsEmpty())
 	{
-		FRay3d LocalRay(
-			(FVector3d)MeshTransforms[MeshIndex].InverseTransformPosition(Pos.WorldRay.Origin),
-			(FVector3d)MeshTransforms[MeshIndex].InverseTransformVector(Pos.WorldRay.Direction));
+		// Nothing to add or modify.
+	}
+	else if (ShouldRestartSelection() 
+		|| (ExistingSelectionIndex == IndexConstants::InvalidID && !ShouldRemoveFromSelection()))
+	{
+		// Make a new selection object
+		NewSelections.Emplace();
+		NewSelections.Last().Target = Targets[HitAssetID];
+		NewSelections.Last().Type = ElementType;
+		NewSelections.Last().SelectedIDs.Append(NewIDs);
+	}
+	else if (ExistingSelectionIndex != IndexConstants::InvalidID)
+	{
+		// Modify the existing selection object
+		ModifyExistingSelection(NewSelections[ExistingSelectionIndex].SelectedIDs, NewIDs);
 
-		double RayT = 0;
-		if (MeshSpatials[MeshIndex]->FindNearestHitTriangle(LocalRay, RayT, HitTid))
+		// Object may end up empty due to subtraction or toggle, in which case it needs to be removed.
+		if (NewSelections[ExistingSelectionIndex].IsEmpty())
 		{
-			CurrentSelection.Mesh = MeshSpatials[MeshIndex]->GetMesh();
-			CurrentSelection.TopologyTimestamp = CurrentSelection.Mesh->GetTopologyChangeStamp();
-			CurrentSelectionIndex = MeshIndex;
-			break;
+			NewSelections.RemoveAt(ExistingSelectionIndex);
 		}
 	}
-
-	if (HitTid != IndexConstants::InvalidID)
+	else
 	{
-		if (Mode == EUVEditorMeshSelectionMode::Component)
-		{
-			TRACE_CPUPROFILER_EVENT_SCOPE(Component);
-
-			FMeshConnectedComponents MeshSelectedComponent(CurrentSelection.Mesh);
-			TArray<int32> SeedTriangles;
-			SeedTriangles.Add(HitTid);
-			MeshSelectedComponent.FindTrianglesConnectedToSeeds(SeedTriangles);
-			ensure(MeshSelectedComponent.Components.Num() == 1); // Expect each triangle to only be in a single component
-			Result.Append(MoveTemp(MeshSelectedComponent.Components[0].Indices));
-		}
-		else if (Mode == EUVEditorMeshSelectionMode::Edge)
-		{
-			TRACE_CPUPROFILER_EVENT_SCOPE(Edge);
-			// TODO: We'll need the ability to hit occluded triangles to see if there is a better edge to snap to.
-
-			// Try to snap to one of the edges.
-			FIndex3i Eids = CurrentSelection.Mesh->GetTriEdges(HitTid);
+		// The only way we can get here is if didn't have an existing selection and were trying
+		// to remove selection, in which case we do nothing.
+		ensure(ExistingSelectionIndex == IndexConstants::InvalidID && ShouldRemoveFromSelection());
+	}
 
-			FGeometrySet3 GeometrySet;
-			for (int i = 0; i < 3; ++i)
-			{
-				FIndex2i Vids = CurrentSelection.Mesh->GetEdgeV(Eids[i]);
-				FPolyline3d Polyline(CurrentSelection.Mesh->GetVertex(Vids.A), CurrentSelection.Mesh->GetVertex(Vids.B));
-				GeometrySet.AddCurve(Eids[i], Polyline);
-			}
+	SelectionAPI->SetSelections(NewSelections, true, true); // broadcast and emit
+}
 
-			FRay3d WorldRay((FVector3d)Pos.WorldRay.Origin, (FVector3d)Pos.WorldRay.Direction);
-			FGeometrySet3::FNearest Nearest;
-			if (GeometrySet.FindNearestCurveToRay(WorldRay, Nearest, 
-				[this](const FVector3d& Position1, const FVector3d& Position2) {
-					return ToolSceneQueriesUtil::PointSnapQuery(CameraState,
-						Position1, Position2,
-						ToolSceneQueriesUtil::GetDefaultVisualAngleSnapThreshD()); }))
-			{
-				Result = TSet<int32>{Nearest.ID};
-			}
-		}
-		else if (Mode == EUVEditorMeshSelectionMode::Vertex)
+bool UUVEditorMeshSelectionMechanic::GetHitTid(const FInputDeviceRay& ClickPos, 
+	int32& TidOut, int32& AssetIDOut, int32* ExistingSelectionObjectIndexOut)
+{
+	auto RayCastSpatial = [this, &ClickPos, &TidOut, &AssetIDOut](int32 AssetID) {
+		double RayT = 0;
+		if (MeshSpatials[AssetID]->FindNearestHitTriangle(ClickPos.WorldRay, RayT, TidOut))
 		{
-			TRACE_CPUPROFILER_EVENT_SCOPE(Vertex);
-			// TODO: Improve this to handle super narrow, sliver triangles better, where testing near vertices can be difficult.
-
-			// Try to snap to one of the vertices
-			FIndex3i Vids = CurrentSelection.Mesh->GetTriangle(HitTid);
+			AssetIDOut = AssetID;
+			return true;
+		}
+		return false;
+	};
 
-			FGeometrySet3 GeometrySet;
-			for (int i = 0; i < 3; ++i)
-			{
-				GeometrySet.AddPoint(Vids[i], CurrentSelection.Mesh->GetTriVertex(HitTid, i));
-			}
-			
-			FRay3d WorldRay((FVector3d)Pos.WorldRay.Origin, (FVector3d)Pos.WorldRay.Direction);
-			FGeometrySet3::FNearest Nearest;
-			if (GeometrySet.FindNearestPointToRay(WorldRay, Nearest,
-				[this](const FVector3d& Position1, const FVector3d& Position2) {
-					return ToolSceneQueriesUtil::PointSnapQuery(CameraState,
-						Position1, Position2,
-						ToolSceneQueriesUtil::GetDefaultVisualAngleSnapThreshD()); }))
+	// Try raycasting the selected meshes first
+	TArray<bool> SpatialTriedFlags;
+	SpatialTriedFlags.SetNum(MeshSpatials.Num());
+	const TArray<FUVToolSelection>& Selections = SelectionAPI->GetSelections();
+	for (int32 SelectionIndex = 0; SelectionIndex < Selections.Num(); ++SelectionIndex)
+	{
+		const FUVToolSelection& Selection = Selections[SelectionIndex];
+		if (ensure(Selection.Target.IsValid() && Selection.Target->AssetID < MeshSpatials.Num()))
+		{
+			if (RayCastSpatial(Selection.Target->AssetID))
 			{
-				Result = TSet<int32>{Nearest.ID};
+				if (ExistingSelectionObjectIndexOut)
+				{
+					*ExistingSelectionObjectIndexOut = SelectionIndex;
+				}
+				return true;
 			}
+			SpatialTriedFlags[Selection.Target->AssetID] = true;
 		}
-		else if (Mode == EUVEditorMeshSelectionMode::Triangle)
-		{
-			TRACE_CPUPROFILER_EVENT_SCOPE(Triangle);
+	}
 
-			Result = TSet<int32>{HitTid};
-		}
-		else if (Mode == EUVEditorMeshSelectionMode::Mesh)
-		{
-			TRACE_CPUPROFILER_EVENT_SCOPE(Mesh);
+	if (ExistingSelectionObjectIndexOut)
+	{
+		*ExistingSelectionObjectIndexOut = IndexConstants::InvalidID;
+	}
 
-			for (int32 Tid : CurrentSelection.Mesh->TriangleIndicesItr())
-			{
-				Result.Add(Tid);
-			}
+	// Try raycasting the other meshes
+	for (int32 AssetID = 0; AssetID < MeshSpatials.Num(); ++AssetID)
+	{
+		if (SpatialTriedFlags[AssetID])
+		{
+			continue;
 		}
-		else
+		if (RayCastSpatial(AssetID))
 		{
-			checkNoEntry();
+			return true;
 		}
 	}
 
-	return Result;
+	return false;
 }
 
 void UUVEditorMeshSelectionMechanic::OnDragRectangleStarted()
 {
+	using namespace UVEditorMeshSelectionMechanicLocals;
+
 	TRACE_CPUPROFILER_EVENT_SCOPE(MeshSelectionMechanic_OnDragRectangleStarted); // Mark start of drag sequence
 	
-	if (CurrentSelection.Type != UVEditorMeshSelectionMechanicLocals::ToCompatibleDynamicMeshSelectionType(SelectionMode))
-	{
-		ClearCurrentSelection();	
-	}
-	
-	PreDragSelection = CurrentSelection;
-}
+	PreDragSelections = SelectionAPI->GetSelections();
+	SelectionAPI->BeginChange();
 
-void UUVEditorMeshSelectionMechanic::OnDragRectangleChanged(const FCameraRectangle& CurrentRectangle)
-{
-	TRACE_CPUPROFILER_EVENT_SCOPE(MeshSelectionMechanic_OnDragRectangleChanged);
-	
-	auto FindIDsInsideCurrentRectangle = [this](auto&& AddIDsToSelection)
+	AssetIDToPreDragSelection.Reset();
+	AssetIDToPreDragSelection.SetNumZeroed(Targets.Num());
+	FUVToolSelection::EType ExpectedSelectionType = ToCompatibleDynamicMeshSelectionType(SelectionMode);
+	if (SelectionAPI->HaveSelections() 
+		&& SelectionAPI->GetSelectionsType() == ExpectedSelectionType)
 	{
-		TSet<int32> RectangleSelectedIDs;
-		
-		if (CurrentSelectionIndex == IndexConstants::InvalidID)
+		for (FUVToolSelection& Selection : PreDragSelections)
 		{
-			for (int32 MeshIndex = 0; MeshIndex < MeshSpatials.Num(); ++MeshIndex)
+			if (ensure(Selection.Type == ExpectedSelectionType))
 			{
-				// TODO UUVEditorMeshSelectionMechanic is currently assuming that the provided mesh transform is identity, which is
-				//  the case in the UV editor, this restriction should be lifted. When we do this we should also apply the
-				//  (inverse) transform to the CurrentRectangle query rays. Search :ApplyTransformToQuery
-				ensure(MeshTransforms[MeshIndex].Identical(&FTransform::Identity, 0));
-				AddIDsToSelection(*MeshSpatials[MeshIndex], RectangleSelectedIDs);
-	
-				if (!RectangleSelectedIDs.IsEmpty())
-				{
-					// Pick the first selectable mesh for now, maybe we should try to be smarter
-					CurrentSelectionIndex = MeshIndex;
-					CurrentSelection.Mesh = MeshSpatials[CurrentSelectionIndex]->GetMesh();
-					break;
-				}
+				AssetIDToPreDragSelection[Selection.Target->AssetID] = &Selection;
 			}
 		}
-		else
-		{
-			ensure(CurrentSelectionIndex != IndexConstants::InvalidID);
-			ensure(CurrentSelection.Mesh == MeshSpatials[CurrentSelectionIndex]->GetMesh());
+	}
+}
 
-			// TODO See :ApplyTransformToQuery
-			ensure(MeshTransforms[CurrentSelectionIndex].Identical(&FTransform::Identity, 0));
-			AddIDsToSelection(*MeshSpatials[CurrentSelectionIndex], RectangleSelectedIDs);
-		}
+void UUVEditorMeshSelectionMechanic::OnDragRectangleChanged(const FCameraRectangle& CurrentRectangle)
+{
+	TRACE_CPUPROFILER_EVENT_SCOPE(MeshSelectionMechanic_OnDragRectangleChanged);
 
-		return RectangleSelectedIDs;
-	};
 	using namespace UVEditorMeshSelectionMechanicLocals;
+
 	FAxisAlignedBox2d RectangleXY = GetRectangleXY(CurrentRectangle);
-	
-	TSet<int32> RectangleSelectedIDs;
-	if (SelectionMode == EUVEditorMeshSelectionMode::Vertex)
+	TArray<FUVToolSelection> NewSelections;
+	FUVToolSelection::EType SelectionType = ToCompatibleDynamicMeshSelectionType(SelectionMode);
+
+	// Gather IDs in each target
+	for (int32 AssetID = 0; AssetID < Targets.Num(); ++AssetID)
 	{
-		TRACE_CPUPROFILER_EVENT_SCOPE(MeshSelectionMechanic_OnDragRectangleChanged_Vertex);
-	
-		auto AddVertexIDsToSelection = [this, &RectangleXY](const FDynamicMeshAABBTree3& Tree, TSet<int32>& OutSelectedIDs)
+		TArray<int32> RectangleSelectedIDs;
+		const FDynamicMeshAABBTree3& Tree = *MeshSpatials[AssetID];
+
+		if (SelectionMode == ESelectionMode::Vertex)
 		{
-			TArray<int32> SelectedIDs =
-				FindAllIntersectionsAxisAlignedBox2(Tree, RectangleXY, AppendVertexIDs, AppendVertexIDsIfIntersected);
-			OutSelectedIDs.Append(MoveTemp(SelectedIDs));
-		};
-	
-		RectangleSelectedIDs = FindIDsInsideCurrentRectangle(AddVertexIDsToSelection);
-	}
-	else if (SelectionMode == EUVEditorMeshSelectionMode::Edge)
-	{
-		TRACE_CPUPROFILER_EVENT_SCOPE(MeshSelectionMechanic_OnDragRectangleChanged_Edge);
-		
-		auto AddEdgeIDsToSelection = [this, &RectangleXY](const FDynamicMeshAABBTree3& Tree, TSet<int32>& OutSelectedIDs)
+			TRACE_CPUPROFILER_EVENT_SCOPE(MeshSelectionMechanic_OnDragRectangleChanged_Vertex);
+
+			RectangleSelectedIDs = FindAllIntersectionsAxisAlignedBox2(
+				Tree, RectangleXY, AppendVertexIDs, AppendVertexIDsIfIntersected);
+		}
+		else if (SelectionMode == ESelectionMode::Edge)
 		{
-			TArray<int32> SelectedIDs =
-				FindAllIntersectionsAxisAlignedBox2(Tree, RectangleXY, AppendEdgeIDs, AppendEdgeIDsIfIntersected);
-			OutSelectedIDs.Append(MoveTemp(SelectedIDs));
-		};
+			TRACE_CPUPROFILER_EVENT_SCOPE(MeshSelectionMechanic_OnDragRectangleChanged_Edge);
 
-		RectangleSelectedIDs = FindIDsInsideCurrentRectangle(AddEdgeIDsToSelection);
-	}
-	else if (SelectionMode == EUVEditorMeshSelectionMode::Triangle)
-	{
-		TRACE_CPUPROFILER_EVENT_SCOPE(MeshSelectionMechanic_OnDragRectangleChanged_Triangle);
-		
-		auto AddTriangleIDsToSelection = [this, &RectangleXY] (const FDynamicMeshAABBTree3& Tree, TSet<int32>& OutSelectedIDs)
+			RectangleSelectedIDs = FindAllIntersectionsAxisAlignedBox2(
+				Tree, RectangleXY, AppendEdgeIDs, AppendEdgeIDsIfIntersected);
+		}
+		else if (SelectionMode == ESelectionMode::Triangle)
 		{
-			TArray<int32> SelectedIDs =
-				FindAllIntersectionsAxisAlignedBox2(Tree, RectangleXY, AppendTriangleID, AppendTriangleIDIfIntersected);
-			OutSelectedIDs.Append(MoveTemp(SelectedIDs));
-		};
-	
-		RectangleSelectedIDs = FindIDsInsideCurrentRectangle(AddTriangleIDsToSelection);
-	}
-	else if (SelectionMode == EUVEditorMeshSelectionMode::Component)
-	{
-		TRACE_CPUPROFILER_EVENT_SCOPE(MeshSelectionMechanic_OnDragRectangleChanged_Component);
-		
-		auto AddTriangleIDsToSelection = [this, &RectangleXY](const FDynamicMeshAABBTree3& Tree, TSet<int32>& OutSelectedIDs)
+			TRACE_CPUPROFILER_EVENT_SCOPE(MeshSelectionMechanic_OnDragRectangleChanged_Triangle);
+
+			RectangleSelectedIDs = FindAllIntersectionsAxisAlignedBox2(
+				Tree, RectangleXY, AppendTriangleID, AppendTriangleIDIfIntersected);
+		}
+		else if (SelectionMode == ESelectionMode::Island)
 		{
-			TArray<int32> SeedTriangles =
-				FindAllIntersectionsAxisAlignedBox2(Tree, RectangleXY, AppendTriangleID, AppendTriangleIDIfIntersected);
+			TRACE_CPUPROFILER_EVENT_SCOPE(MeshSelectionMechanic_OnDragRectangleChanged_Component);
+
+			TArray<int32> SeedTriangles = FindAllIntersectionsAxisAlignedBox2(
+				Tree, RectangleXY, AppendTriangleID, AppendTriangleIDIfIntersected);
 
 			// TODO(Performance) For large meshes and selections following code is MUCH slower than AABB traversal,
 			//  consider precomputing the connected components an only updating them when the mesh topology changes
@@ -1232,43 +1033,71 @@ void UUVEditorMeshSelectionMechanic::OnDragRectangleChanged(const FCameraRectang
 			MeshSelectedComponent.FindTrianglesConnectedToSeeds(SeedTriangles);
 			for (int ComponentIndex = 0; ComponentIndex < MeshSelectedComponent.Components.Num(); ComponentIndex++)
 			{
-				OutSelectedIDs.Append(MoveTemp(MeshSelectedComponent.Components[ComponentIndex].Indices));
+				RectangleSelectedIDs.Append(MoveTemp(MeshSelectedComponent.Components[ComponentIndex].Indices));
 			}
-		};
-	
-		RectangleSelectedIDs = FindIDsInsideCurrentRectangle(AddTriangleIDsToSelection);
-	}
-	else if (SelectionMode == EUVEditorMeshSelectionMode::Mesh)
-	{
-		TRACE_CPUPROFILER_EVENT_SCOPE(MeshSelectionMechanic_OnDragRectangleChanged_Mesh);
-		
-		auto AddAllMeshTriangleIDsToSelection = [this, &RectangleXY] (const FDynamicMeshAABBTree3& Tree, TSet<int32>& OutSelectedIDs)
+		}
+		else if (SelectionMode == ESelectionMode::Mesh)
 		{
-			TArray<int32> SelectedIDs =
-				FindAllIntersectionsAxisAlignedBox2(Tree, RectangleXY, AppendTriangleID, AppendTriangleIDIfIntersected);
+			TRACE_CPUPROFILER_EVENT_SCOPE(MeshSelectionMechanic_OnDragRectangleChanged_Mesh);
+
+			// TODO: This shouldn't be a "find all". We can return early after the first success
+			// since we're selecting the whole mesh
+			TArray<int32> SelectedIDs = FindAllIntersectionsAxisAlignedBox2(
+				Tree, RectangleXY, AppendTriangleID, AppendTriangleIDIfIntersected);
 			if (!SelectedIDs.IsEmpty())
 			{
 				for (int32 Tid : Tree.GetMesh()->TriangleIndicesItr())
 				{
-					OutSelectedIDs.Add(Tid);
+					RectangleSelectedIDs.Add(Tid);
 				}
 			}
-		};
+		}
+		else
+		{
+			checkSlow(false);
+		}
 
-		RectangleSelectedIDs = FindIDsInsideCurrentRectangle(AddAllMeshTriangleIDsToSelection);
-	}
-	else
-	{
-		checkNoEntry();
+		// See if we have an object in our selection list that corresponds to this asset
+		const FUVToolSelection* PreDragSelection = AssetIDToPreDragSelection[AssetID];
+
+		if (RectangleSelectedIDs.IsEmpty())
+		{
+			if (!ShouldRestartSelection() && PreDragSelection)
+			{
+				// Keep the existing selection object with no modification.
+				NewSelections.Emplace(*PreDragSelection);
+			}
+		}
+		else if (ShouldRestartSelection() || (!PreDragSelection && !ShouldRemoveFromSelection()))
+		{
+			// Make a new selection object
+			NewSelections.Emplace();
+			NewSelections.Last().Target = Targets[AssetID];
+			NewSelections.Last().Type = SelectionType;
+			NewSelections.Last().SelectedIDs.Append(RectangleSelectedIDs);
+		}
+		else if (PreDragSelection)
+		{
+			// Modify the existing selection object
+			FUVToolSelection NewSelection(*PreDragSelection);
+			ModifyExistingSelection(NewSelection.SelectedIDs, RectangleSelectedIDs);
+
+			// The object may become empty from a removal or toggle, in which case don't add it.
+			if (!NewSelection.IsEmpty())
+			{
+				NewSelections.Add(MoveTemp(NewSelection));
+			}
+		}
+		else
+		{
+			// The only way we can get here is if didn't have an existing selection and were trying
+			// to remove selection, in which case we do nothing.
+			ensure(!PreDragSelection && ShouldRemoveFromSelection());
+		}
 	}
 
-	UpdateCurrentSelection(RectangleSelectedIDs, true);
-	
-	// TODO(Performance) With large meshes and selections this call is much slower than AABB traversal (4x)
-	UpdateCentroid();
-	
-	// TODO(Performance) With large meshes and selections this call is MUCH slower than AABB traversal (60x)
-	RebuildDrawnElements(FTransform(GetCurrentSelectionCentroid()));
+	SelectionAPI->SetSelections(NewSelections, false, false);
+	OnDragSelectionChanged.Broadcast();
 }
 
 void UUVEditorMeshSelectionMechanic::OnDragRectangleFinished(const FCameraRectangle& CurrentRectangle, bool bCancelled)
@@ -1279,14 +1108,9 @@ void UUVEditorMeshSelectionMechanic::OnDragRectangleFinished(const FCameraRectan
 	//  dynamically for large meshes
 	OnDragRectangleChanged(CurrentRectangle);
 
-	if (!bCancelled && (PreDragSelection != CurrentSelection))
+	if (!bCancelled)
 	{
-		UpdateCentroid();		
-		EmitSelectionChange(PreDragSelection, CurrentSelection);
-		OnSelectionChanged.Broadcast();
-
-		// Rebuild after broadcast in case the outside world wants to adjust things like color...
-		RebuildDrawnElements(FTransform(GetCurrentSelectionCentroid()));
+		SelectionAPI->EndChangeAndEmitIfModified(true);
 	}
 }
 
@@ -1305,5 +1129,114 @@ void UUVEditorMeshSelectionMechanic::OnUpdateModifierState(int ModifierID, bool
 	}
 }
 
+FInputRayHit UUVEditorMeshSelectionMechanic::BeginHoverSequenceHitTest(const FInputDeviceRay& PressPos)
+{
+	FInputRayHit Hit;
+	if (!bIsEnabled || !bShowHoveredElements || SelectionMode == ESelectionMode::None)
+	{
+		Hit.bHit = false;
+		return Hit;
+	}
+
+	ESelectionMode Mode = SelectionMode;
+	if (Mode != ESelectionMode::Vertex && Mode != ESelectionMode::Edge)
+	{
+		Mode = ESelectionMode::Triangle;
+	}
+
+	// We don't bother with the depth since everything is in the same plane.
+	int32 Tid = IndexConstants::InvalidID;
+	int32 AssetID = IndexConstants::InvalidID;
+	Hit.bHit = GetHitTid(PressPos, Tid, AssetID);
+
+	return Hit;
+}
+
+void UUVEditorMeshSelectionMechanic::OnBeginHover(const FInputDeviceRay& DevicePos)
+{
+}
+
+bool UUVEditorMeshSelectionMechanic::OnUpdateHover(const FInputDeviceRay& DevicePos)
+{
+	using namespace UVEditorMeshSelectionMechanicLocals;
+
+	ESelectionMode Mode = SelectionMode;
+	if (SelectionMode != ESelectionMode::Vertex && SelectionMode != ESelectionMode::Edge)
+	{
+		Mode = ESelectionMode::Triangle;
+	}
+
+	HoverPointSet->Clear();
+	HoverLineSet->Clear();
+	HoverTriangleSet->Clear();
+
+	int32 Tid = IndexConstants::InvalidID;
+	int32 AssetID = IndexConstants::InvalidID;
+	if (!GetHitTid(DevicePos, Tid, AssetID))
+	{
+		return false;
+	}
+
+	FDynamicMesh3* Mesh = Targets[AssetID]->UnwrapCanonical.Get();
+
+	TArray<int32> ConvertedIDs;
+	if (SelectionMode == ESelectionMode::Vertex || SelectionMode == ESelectionMode::Edge)
+	{
+		ConvertToHitElementList(SelectionMode, *Mesh,
+			Tid, CameraState, DevicePos.WorldRay, ConvertedIDs);
+		if (ConvertedIDs.IsEmpty())
+		{
+			// We were too far from a vert or edge, probably.
+			return false;
+		}
+	}
+
+	if (SelectionMode == ESelectionMode::Vertex)
+	{
+		const FVector3d& P = Mesh->GetVertexRef(ConvertedIDs[0]);
+		const FRenderablePoint PointToRender(P,
+			FUVEditorUXSettings::SelectionHoverTriangleWireframeColor,
+			FUVEditorUXSettings::SelectionPointThickness);
+
+		HoverPointSet->AddPoint(PointToRender);
+	}
+	else if (SelectionMode == ESelectionMode::Edge)
+	{
+		const FIndex2i EdgeVids = Mesh->GetEdgeV(ConvertedIDs[0]);
+		const FVector& A = Mesh->GetVertexRef(EdgeVids.A);
+		const FVector& B = Mesh->GetVertexRef(EdgeVids.B);
+
+		HoverLineSet->AddLine(A, B,
+			FUVEditorUXSettings::SelectionHoverTriangleWireframeColor,
+			FUVEditorUXSettings::SelectionLineThickness,
+			FUVEditorUXSettings::SelectionHoverWireframeDepthBias);
+	}
+	else
+	{
+		const FIndex3i Vids = Mesh->GetTriangle(Tid);
+		const FVector& A = Mesh->GetVertex(Vids[0]);
+		const FVector& B = Mesh->GetVertex(Vids[1]);
+		const FVector& C = Mesh->GetVertex(Vids[2]);
+
+		HoverLineSet->AddLine(A, B, FUVEditorUXSettings::SelectionHoverTriangleWireframeColor,
+			FUVEditorUXSettings::SelectionLineThickness, FUVEditorUXSettings::SelectionHoverWireframeDepthBias);
+		HoverLineSet->AddLine(B, C, FUVEditorUXSettings::SelectionHoverTriangleWireframeColor,
+			FUVEditorUXSettings::SelectionLineThickness, FUVEditorUXSettings::SelectionHoverWireframeDepthBias);
+		HoverLineSet->AddLine(C, A, FUVEditorUXSettings::SelectionHoverTriangleWireframeColor,
+			FUVEditorUXSettings::SelectionLineThickness, FUVEditorUXSettings::SelectionHoverWireframeDepthBias);
+		HoverTriangleSet->AddTriangle(A, B, C, FVector::ZAxisVector,
+			FUVEditorUXSettings::SelectionHoverTriangleFillColor, HoverTriangleSetMaterial);
+	}
+
+	return true;
+}
+
+void UUVEditorMeshSelectionMechanic::OnEndHover()
+{
+	HoverPointSet->Clear();
+	HoverLineSet->Clear();
+	HoverTriangleSet->Clear();
+}
+
 #undef LOCTEXT_NAMESPACE
 
diff --git a/Engine/Plugins/Experimental/UVEditor/Source/UVEditorTools/Private/Selection/UVToolSelection.cpp b/Engine/Plugins/Experimental/UVEditor/Source/UVEditorTools/Private/Selection/UVToolSelection.cpp
new file mode 100644
--- /dev/null
+++ b/Engine/Plugins/Experimental/UVEditor/Source/UVEditorTools/Private/Selection/UVToolSelection.cpp
@@ -0,0 +1,60 @@
+// Copyright Epic Games, Inc. All Rights Reserved.
+
+#include "Selection/UVToolSelection.h"
+
+#define LOCTEXT_NAMESPACE "UUVToolSelection"
+
+using namespace UE::Geometry;
+
+void FUVToolSelection::SaveStableEdgeIdentifiers(const FDynamicMesh3& Mesh)
+{
+	if (Type == EType::Edge)
+	{
+		StableEdgeIDs.InitializeFromEdgeIDs(Mesh, SelectedIDs);
+	}
+}
+
+void FUVToolSelection::RestoreFromStableEdgeIdentifiers(const FDynamicMesh3& Mesh)
+{
+	if (Type == EType::Edge)
+	{
+		StableEdgeIDs.GetEdgeIDs(Mesh, SelectedIDs);
+	}
+}
+
+bool FUVToolSelection::AreElementsPresentInMesh(FDynamicMesh3& Mesh) const
+{
+	switch (Type)
+	{
+	case EType::Vertex:
+		for (int32 Vid : SelectedIDs)
+		{
+			if (!Mesh.IsVertex(Vid))
+			{
+				return false;
+			}
+		}
+		break;
+	case EType::Edge:
+		for (int32 Eid : SelectedIDs)
+		{
+			if (!Mesh.IsEdge(Eid))
+			{
+				return false;
+			}
+		}
+		break;
+	case EType::Triangle:
+		for (int32 Tid : SelectedIDs)
+		{
+			if (!Mesh.IsTriangle(Tid))
+			{
+				return false;
+			}
+		}
+		break;
+	}
+	return true;
+}
+
+#undef LOCTEXT_NAMESPACE
\ No newline at end of file
diff --git a/Engine/Plugins/Experimental/UVEditor/Source/UVEditorTools/Private/Selection/UVToolSelectionAPI.cpp b/Engine/Plugins/Experimental/UVEditor/Source/UVEditorTools/Private/Selection/UVToolSelectionAPI.cpp
new file mode 100644
--- /dev/null
+++ b/Engine/Plugins/Experimental/UVEditor/Source/UVEditorTools/Private/Selection/UVToolSelectionAPI.cpp
@@ -0,0 +1,559 @@
+// Copyright Epic Games, Inc. All Rights Reserved.
+
+#include "Selection/UVToolSelectionAPI.h"
+
+#include "DynamicMesh/DynamicMesh3.h"
+#include "MeshOpPreviewHelpers.h" // UMeshOpPreviewWithBackgroundCompute
+#include "Selection/UVEditorMeshSelectionMechanic.h"
+#include "Selection/UVToolSelectionHighlightMechanic.h"
+#include "UVEditorMechanicAdapterTool.h"
+
+#define LOCTEXT_NAMESPACE "UUVIslandConformalUnwrapAction"
+
+using namespace UE::Geometry;
+
+namespace UVToolSelectionAPILocals
+{
+	FText SelectionChangeTransactionName = LOCTEXT("SelectionChangeTransaction", "UV Selection Change");
+
+	auto DoSelectionSetsDiffer(const TArray<FUVToolSelection>& OldSelections, const TArray<FUVToolSelection>& NewSelections)
+	{
+		if (NewSelections.Num() != OldSelections.Num())
+		{
+			return true;
+		}
+
+		for (const FUVToolSelection& Selection : NewSelections)
+		{
+			// Find the selection that points to the same target.
+			const FUVToolSelection* FoundSelection = OldSelections.FindByPredicate(
+				[&Selection](const FUVToolSelection& OldSelection) { return OldSelection.Target == Selection.Target; }
+			);
+			if (!FoundSelection || *FoundSelection != Selection)
+			{
+				return true;
+			}
+		}
+		return false;
+	};
+}
+
+void UUVToolSelectionAPI::Initialize(
+	UInteractiveToolManager* ToolManagerIn,
+	UWorld* UnwrapWorld, UInputRouter* UnwrapInputRouterIn, 
+	UUVToolLivePreviewAPI* LivePreviewAPI,
+	UUVToolEmitChangeAPI* EmitChangeAPIIn)
+{
+	UnwrapInputRouter = UnwrapInputRouterIn;
+	EmitChangeAPI = EmitChangeAPIIn;
+
+	MechanicAdapter = NewObject<UUVEditorMechanicAdapterTool>();
+	MechanicAdapter->ToolManager = ToolManagerIn;
+
+	HighlightMechanic = NewObject<UUVToolSelectionHighlightMechanic>();
+	HighlightMechanic->Setup(MechanicAdapter);
+	HighlightMechanic->Initialize(UnwrapWorld, LivePreviewAPI->GetLivePreviewWorld());
+
+	SelectionMechanic = NewObject<UUVEditorMeshSelectionMechanic>();
+	SelectionMechanic->Setup(MechanicAdapter);
+	SelectionMechanic->Initialize(UnwrapWorld, this);
+	UnwrapInputRouter->RegisterSource(MechanicAdapter);
+}
+
+void UUVToolSelectionAPI::SetTargets(const TArray<TObjectPtr<UUVEditorToolMeshInput>>& TargetsIn)
+{
+	Targets = TargetsIn;
+	SelectionMechanic->SetTargets(Targets);
+
+	for (TObjectPtr<UUVEditorToolMeshInput> Target : Targets)
+	{
+		Target->OnCanonicalModified.AddWeakLambda(this, [this](UUVEditorToolMeshInput* Target,
+			const UUVEditorToolMeshInput::FCanonicalModifiedInfo ModifiedInfo)
+			{
+				FUVToolSelection* TargetSelection = CurrentSelections.FindByPredicate(
+					[Target](const FUVToolSelection& CandidateSelection) { return CandidateSelection.Target == Target; });
+				if (TargetSelection && TargetSelection->Type == FUVToolSelection::EType::Edge)
+				{
+					TargetSelection->RestoreFromStableEdgeIdentifiers(*Target->UnwrapCanonical);
+				}
+			});
+	}
+}
+
+void UUVToolSelectionAPI::Shutdown()
+{
+	HighlightMechanic->Shutdown();
+	SelectionMechanic->Shutdown();
+
+	MechanicAdapter->Shutdown(EToolShutdownType::Completed);
+	HighlightMechanic = nullptr;
+	SelectionMechanic = nullptr;
+	MechanicAdapter = nullptr;
+
+	if (UnwrapInputRouter.IsValid())
+	{
+		// Make sure that we stop any captures that our mechanics may have, then remove them from
+		// the input router.
+		UnwrapInputRouter->ForceTerminateSource(MechanicAdapter);
+		UnwrapInputRouter->DeregisterSource(MechanicAdapter);
+		UnwrapInputRouter = nullptr;
+	}
+
+	for (TObjectPtr<UUVEditorToolMeshInput> Target : Targets)
+	{
+		Target->OnCanonicalModified.RemoveAll(this);
+	}
+	Targets.Empty();
+}
+
+void UUVToolSelectionAPI::ClearSelections(bool bBroadcast, bool bEmitChange)
+{
+	SetSelections(TArray<FUVToolSelection>(), bBroadcast, bEmitChange);
+}
+
+void UUVToolSelectionAPI::OnToolEnded(UInteractiveTool* DeadTool)
+{
+	if (SelectionMechanic)
+	{
+		SetSelectionMechanicOptions(FSelectionMechanicOptions());
+		SelectionMechanic->SetIsEnabled(false);
+	}
+	if (HighlightMechanic)
+	{
+		SetHighlightVisible(false, false);
+		SetHighlightOptions(FHighlightOptions());
+	}
+
+	OnSelectionChanged.RemoveAll(DeadTool);
+	OnPreSelectionChange.RemoveAll(DeadTool);
+	OnDragSelectionChanged.RemoveAll(DeadTool);
+}
+
+void UUVToolSelectionAPI::SetSelections(const TArray<FUVToolSelection>& SelectionsIn, bool bBroadcast, bool bEmitChange)
+{
+	using namespace UVToolSelectionAPILocals;
+
+	bCachedUnwrapSelectionCentroidValid = false;
+
+	TArray<FUVToolSelection> NewSelections;
+	for (const FUVToolSelection& NewSelection : SelectionsIn)
+	{
+		if (ensure(NewSelection.Target.IsValid() && NewSelection.Target->IsValid())
+			// All of the selections should match type
+			&& ensure(NewSelections.Num() == 0 || NewSelection.Type == NewSelections[0].Type)
+			// Selection must not be empty, unless it's an edge selection stored as stable identifiers
+			&& ensure(!NewSelection.IsEmpty() 
+				|| (NewSelection.Type == FUVToolSelection::EType::Edge && NewSelection.HasStableEdgeIdentifiers()))
+			// Shouldn't have selection objects pointing to same target
+			&& ensure(!NewSelections.FindByPredicate(
+				[&NewSelection](const FUVToolSelection& ExistingSelectionElement) {
+					return NewSelection.Target == ExistingSelectionElement.Target;
+				})))
+		{
+			NewSelections.Add(NewSelection);
+
+			if (NewSelection.Target.IsValid() && ensure(NewSelection.Target->UnwrapCanonical))
+			{
+				checkSlow(NewSelection.AreElementsPresentInMesh(*NewSelection.Target->UnwrapCanonical));
+
+				if (NewSelection.Type == FUVToolSelection::EType::Edge)
+				{
+					NewSelections.Last().SaveStableEdgeIdentifiers(*NewSelection.Target->UnwrapCanonical);
+				}
+			}
+		}
+	}
+
+	if (!DoSelectionSetsDiffer(CurrentSelections, NewSelections))
+	{
+		return;
+	}
+
+	if (bEmitChange)
+	{
+		EmitChangeAPI->BeginUndoTransaction(SelectionChangeTransactionName);
+	}
+
+	if (bBroadcast)
+	{
+		OnPreSelectionChange.Broadcast(bEmitChange);
+	}
+
+	TUniquePtr<FSelectionChange> SelectionChange;
+	if (bEmitChange)
+	{
+		SelectionChange = MakeUnique<FSelectionChange>();
+		SelectionChange->SetBefore(MoveTemp(CurrentSelections));
+	}
+
+	CurrentSelections = MoveTemp(NewSelections);
+
+	if (bEmitChange)
+	{
+		SelectionChange->SetAfter(CurrentSelections);
+		EmitChangeAPI->EmitToolIndependentChange(this, MoveTemp(SelectionChange),
+			SelectionChangeTransactionName);
+	}
+
+	if (HighlightOptions.bAutoUpdateUnwrap)
+	{
+		FTransform Transform = FTransform::Identity;
+		if (HighlightOptions.bUseCentroidForUnwrapAutoUpdate)
+		{
+			Transform = FTransform(GetUnwrapSelectionCentroid());
+		}
+		RebuildUnwrapHighlight(Transform);
+	}
+	if (HighlightOptions.bAutoUpdateApplied)
+	{
+		RebuildAppliedPreviewHighlight();
+	}
+
+	if (bBroadcast)
+	{
+		OnSelectionChanged.Broadcast(bEmitChange);
+	}
+
+	if (bEmitChange)
+	{
+		EmitChangeAPI->EndUndoTransaction();
+	}
+}
+
+void UUVToolSelectionAPI::SetSelectionMechanicOptions(const FSelectionMechanicOptions& Options)
+{
+	SelectionMechanic->SetShowHoveredElements(Options.bShowHoveredElements);
+}
+
+void UUVToolSelectionAPI::SetSelectionMechanicEnabled(bool bIsEnabled)
+{
+	SelectionMechanic->SetIsEnabled(bIsEnabled);
+}
+
+void UUVToolSelectionAPI::SetSelectionMechanicMode(EUVEditorSelectionMode Mode,
+	const FSelectionMechanicModeChangeOptions& Options)
+{
+	SelectionMechanic->SetSelectionMode(Mode, Options);
+}
+
+FVector3d UUVToolSelectionAPI::GetUnwrapSelectionCentroid(bool bForceRecalculate)
+{
+	if (bCachedUnwrapSelectionCentroidValid && !bForceRecalculate)
+	{
+		return CachedUnwrapSelectionCentroid;
+	}
+
+	FVector3d& Centroid = CachedUnwrapSelectionCentroid;
+
+	Centroid = FVector3d::Zero();
+	double Divisor = 0;
+	for (const FUVToolSelection& Selection : GetSelections())
+	{
+		FDynamicMesh3* Mesh = Selection.Target->UnwrapCanonical.Get();
+		if (Selection.Type == FUVToolSelection::EType::Edge)
+		{
+			for (int32 Eid : Selection.SelectedIDs)
+			{
+				Centroid += Mesh->GetEdgePoint(Eid, 0.5);
+			}
+		}
+		else if (Selection.Type == FUVToolSelection::EType::Triangle)
+		{
+			for (int32 Tid : Selection.SelectedIDs)
+			{
+				Centroid += Mesh->GetTriCentroid(Tid);
+			}
+		}
+		else if (Selection.Type == FUVToolSelection::EType::Vertex)
+		{
+			for (int32 Vid : Selection.SelectedIDs)
+			{
+				Centroid += Mesh->GetVertex(Vid);
+			}
+		}
+		else
+		{
+			ensure(false);
+		}
+
+		Divisor += Selection.SelectedIDs.Num();
+	}
+
+	if (Divisor > 0)
+	{
+		Centroid /= Divisor;
+	}
+
+	bCachedUnwrapSelectionCentroidValid = true;
+	return Centroid;
+}
+
+
+void UUVToolSelectionAPI::SetHighlightVisible(bool bUnwrapHighlightVisible, bool bAppliedHighlightVisible, bool bRebuild)
+{
+	HighlightMechanic->SetIsVisible(bUnwrapHighlightVisible, bAppliedHighlightVisible);
+	if (bRebuild)
+	{
+		ClearHighlight(!bUnwrapHighlightVisible, !bAppliedHighlightVisible);
+		if (bUnwrapHighlightVisible)
+		{
+			FTransform UnwrapTransform = FTransform::Identity;
+			if (HighlightOptions.bUseCentroidForUnwrapAutoUpdate)
+			{
+				UnwrapTransform = FTransform(GetUnwrapSelectionCentroid());
+			}
+			RebuildUnwrapHighlight(UnwrapTransform);
+		}
+		if (bAppliedHighlightVisible)
+		{
+			RebuildAppliedPreviewHighlight();
+		}
+	}
+}
+
+void UUVToolSelectionAPI::SetHighlightOptions(const FHighlightOptions& Options)
+{
+	HighlightOptions = Options;
+
+	HighlightMechanic->SetEnablePairedEdgeHighlights(Options.bShowPairedEdgeHighlights);
+}
+
+void UUVToolSelectionAPI::ClearHighlight(bool bClearForUnwrap, bool bClearForAppliedPreview)
+{
+	if (bClearForUnwrap)
+	{
+		HighlightMechanic->RebuildUnwrapHighlight(TArray<FUVToolSelection>(), FTransform::Identity);
+	}
+	if (bClearForAppliedPreview)
+	{
+		HighlightMechanic->RebuildAppliedHighlightFromUnwrapSelection(TArray<FUVToolSelection>());
+	}
+}
+
+void UUVToolSelectionAPI::RebuildUnwrapHighlight(const FTransform& StartTransform)
+{
+	// Unfortunately even when tids and vids correspond between unwrap and canonical unwraps, eids
+	// may differ, so we need to convert these over to preview unwrap eids if we're using previews.
+	if (HighlightOptions.bBaseHighlightOnPreviews 
+		&& GetSelectionsType() == FUVToolSelection::EType::Edge)
+	{
+		// TODO: We should probably have some function like "GetSelectionsInPreview(bool bForceRebuild)" 
+		// that is publicly available and does caching.
+		TArray<FUVToolSelection> ConvertedSelections = CurrentSelections;
+		for (FUVToolSelection& ConvertedSelection : ConvertedSelections)
+		{
+			ConvertedSelection.SaveStableEdgeIdentifiers(*ConvertedSelection.Target->UnwrapCanonical);
+			ConvertedSelection.RestoreFromStableEdgeIdentifiers(*ConvertedSelection.Target->UnwrapPreview->PreviewMesh->GetMesh());
+		}
+		HighlightMechanic->RebuildUnwrapHighlight(ConvertedSelections, StartTransform,
+			HighlightOptions.bBaseHighlightOnPreviews);
+	}
+	else
+	{
+		HighlightMechanic->RebuildUnwrapHighlight(CurrentSelections, StartTransform,
+			HighlightOptions.bBaseHighlightOnPreviews);
+	}
+}
+
+void UUVToolSelectionAPI::SetUnwrapHighlightTransform(const FTransform& NewTransform)
+{
+	HighlightMechanic->SetUnwrapHighlightTransform(NewTransform, 
+		HighlightOptions.bShowPairedEdgeHighlights, HighlightOptions.bBaseHighlightOnPreviews);
+}
+
+FTransform UUVToolSelectionAPI::GetUnwrapHighlightTransform()
+{
+	return HighlightMechanic->GetUnwrapHighlightTransform();
+}
+
+void UUVToolSelectionAPI::RebuildAppliedPreviewHighlight()
+{
+	// Unfortunately even when tids and vids correspond between unwrap and canonical unwraps, eids
+	// may differ, so we need to convert these over to preview unwrap eids if we're using previews.
+	if (HighlightOptions.bBaseHighlightOnPreviews
+		&& GetSelectionsType() == FUVToolSelection::EType::Edge)
+	{
+		TArray<FUVToolSelection> ConvertedSelections = CurrentSelections;
+		for (FUVToolSelection& ConvertedSelection : ConvertedSelections)
+		{
+			ConvertedSelection.SaveStableEdgeIdentifiers(*ConvertedSelection.Target->UnwrapCanonical);
+			ConvertedSelection.RestoreFromStableEdgeIdentifiers(*ConvertedSelection.Target->UnwrapPreview->PreviewMesh->GetMesh());
+		}
+		HighlightMechanic->RebuildAppliedHighlightFromUnwrapSelection(ConvertedSelections,
+			HighlightOptions.bBaseHighlightOnPreviews);
+	}
+	else
+	{
+		HighlightMechanic->RebuildAppliedHighlightFromUnwrapSelection(CurrentSelections,
+			HighlightOptions.bBaseHighlightOnPreviews);
+	}
+}
+
+void UUVToolSelectionAPI::Render(IToolsContextRenderAPI* RenderAPI)
+{
+	if (SelectionMechanic)
+	{
+		SelectionMechanic->Render(RenderAPI);
+	}
+}
+void UUVToolSelectionAPI::DrawHUD(FCanvas* Canvas, IToolsContextRenderAPI* RenderAPI)
+{
+	if (SelectionMechanic)
+	{
+		SelectionMechanic->DrawHUD(Canvas, RenderAPI);
+	}
+}
+
+void UUVToolSelectionAPI::BeginChange()
+{
+	PendingSelectionChange = MakeUnique<FSelectionChange>();
+	FSelectionChange* CastSelectionChange = static_cast<FSelectionChange*>(PendingSelectionChange.Get());
+	CastSelectionChange->SetBefore(GetSelections());
+}
+
+bool UUVToolSelectionAPI::EndChangeAndEmitIfModified(bool bBroadcast)
+{
+	using namespace UVToolSelectionAPILocals;
+
+	if (!PendingSelectionChange)
+	{
+		return false;
+	}
+
+	FSelectionChange* CastSelectionChange = static_cast<FSelectionChange*>(PendingSelectionChange.Get());
+	
+	// See if the selection has changed
+	if (!DoSelectionSetsDiffer(CastSelectionChange->GetBefore(), GetSelections()))
+	{
+		PendingSelectionChange.Reset();
+		return false;
+	}
+
+	EmitChangeAPI->BeginUndoTransaction(SelectionChangeTransactionName);
+	if (bBroadcast)
+	{
+		OnPreSelectionChange.Broadcast(true);
+	}
+	CastSelectionChange->SetAfter(GetSelections());
+	EmitChangeAPI->EmitToolIndependentChange(this, MoveTemp(PendingSelectionChange),
+		SelectionChangeTransactionName);
+	PendingSelectionChange.Reset();
+
+	if (bBroadcast)
+	{
+		OnSelectionChanged.Broadcast(true);
+	}
+	EmitChangeAPI->EndUndoTransaction();
+
+	return true;
+}
+
+void UUVToolSelectionAPI::FSelectionChange::SetBefore(TArray<FUVToolSelection> SelectionsIn)
+{
+	Before = MoveTemp(SelectionsIn);
+	if (bUseStableUnwrapCanonicalIDsForEdges)
+	{
+		for (FUVToolSelection& Selection : Before)
+		{
+			if (Selection.Type == FUVToolSelection::EType::Edge
+				&& Selection.Target.IsValid()
+				&& Selection.Target->UnwrapCanonical)
+			{
+				Selection.SaveStableEdgeIdentifiers(*Selection.Target->UnwrapCanonical);
+				Selection.SelectedIDs.Empty(); // Don't need to store since we'll always restore
+			}
+		}
+	}
+}
+
+void UUVToolSelectionAPI::FSelectionChange::SetAfter(TArray<FUVToolSelection> SelectionsIn)
+{
+	After = MoveTemp(SelectionsIn);
+	if (bUseStableUnwrapCanonicalIDsForEdges)
+	{
+		for (FUVToolSelection& Selection : After)
+		{
+			if (Selection.Type == FUVToolSelection::EType::Edge
+				&& Selection.Target.IsValid()
+				&& Selection.Target->UnwrapCanonical)
+			{
+				Selection.SaveStableEdgeIdentifiers(*Selection.Target->UnwrapCanonical);
+				Selection.SelectedIDs.Empty(); // Don't need to store since we'll always restore
+			}
+		}
+	}
+}
+
+TArray<FUVToolSelection> UUVToolSelectionAPI::FSelectionChange::GetBefore() const
+{
+	TArray<FUVToolSelection> SelectionsOut;
+	for (const FUVToolSelection& Selection : Before)
+	{
+		if (bUseStableUnwrapCanonicalIDsForEdges
+			&& Selection.Type == FUVToolSelection::EType::Edge
+			&& Selection.Target.IsValid()
+			&& Selection.Target->UnwrapCanonical)
+		{
+			FUVToolSelection UnpackedSelection = Selection;
+			UnpackedSelection.RestoreFromStableEdgeIdentifiers(*Selection.Target->UnwrapCanonical);
+			SelectionsOut.Add(UnpackedSelection);
+		}
+		else
+		{
+			SelectionsOut.Add(Selection);
+		}
+	}
+	return SelectionsOut;
+}
+
+void UUVToolSelectionAPI::FSelectionChange::Apply(UObject* Object) 
+{
+	UUVToolSelectionAPI* SelectionAPI = Cast<UUVToolSelectionAPI>(Object);
+	if (SelectionAPI)
+	{
+		if (bUseStableUnwrapCanonicalIDsForEdges)
+		{
+			for (FUVToolSelection& Selection : After)
+			{
+				if (Selection.Type == FUVToolSelection::EType::Edge
+					&& Selection.Target.IsValid()
+					&& Selection.Target->UnwrapCanonical)
+				{
+					Selection.RestoreFromStableEdgeIdentifiers(*Selection.Target->UnwrapCanonical);
+				}
+			}
+		}
+
+		SelectionAPI->SetSelections(After, true, false);
+	}
+}
+
+void UUVToolSelectionAPI::FSelectionChange::Revert(UObject* Object)
+{
+	UUVToolSelectionAPI* SelectionAPI = Cast<UUVToolSelectionAPI>(Object);
+	if (SelectionAPI)
+	{
+		if (bUseStableUnwrapCanonicalIDsForEdges)
+		{
+			for (FUVToolSelection& Selection : Before)
+			{
+				if (Selection.Type == FUVToolSelection::EType::Edge
+					&& Selection.Target.IsValid()
+					&& Selection.Target->UnwrapCanonical)
+				{
+					Selection.RestoreFromStableEdgeIdentifiers(*Selection.Target->UnwrapCanonical);
+				}
+			}
+		}
+
+		SelectionAPI->SetSelections(Before, true, false);
+	}
+}
+
+FString UUVToolSelectionAPI::FSelectionChange::ToString() const
+{
+	return TEXT("UUVToolSelectionAPI::FSelectionChange");
+}
+
+
+#undef LOCTEXT_NAMESPACE
\ No newline at end of file
diff --git a/Engine/Plugins/Experimental/UVEditor/Source/UVEditorTools/Private/Selection/UVToolSelectionHighlightMechanic.cpp b/Engine/Plugins/Experimental/UVEditor/Source/UVEditorTools/Private/Selection/UVToolSelectionHighlightMechanic.cpp
new file mode 100644
--- /dev/null
+++ b/Engine/Plugins/Experimental/UVEditor/Source/UVEditorTools/Private/Selection/UVToolSelectionHighlightMechanic.cpp
@@ -0,0 +1,450 @@
+// Copyright Epic Games, Inc. All Rights Reserved.
+
+#include "Selection/UVToolSelectionHighlightMechanic.h"
+
+#include "Actions/UVSeamSewAction.h"
+#include "Drawing/LineSetComponent.h"
+#include "Drawing/PointSetComponent.h"
+#include "Drawing/PreviewGeometryActor.h"
+#include "Drawing/TriangleSetComponent.h"
+#include "DynamicMesh/DynamicMesh3.h"
+#include "MeshOpPreviewHelpers.h" // for the preview meshes
+#include "Selection/UVToolSelection.h"
+#include "ToolSetupUtil.h"
+#include "ToolTargets/UVEditorToolMeshInput.h"
+#include "UVEditorUXSettings.h"
+
+#define LOCTEXT_NAMESPACE "UUVToolSelectionHighlightMechanic"
+
+using namespace UE::Geometry;
+
+namespace UVToolSelectionHighlightMechanicLocals
+{
+}
+
+void UUVToolSelectionHighlightMechanic::Initialize(UWorld* UnwrapWorld, UWorld* LivePreviewWorld)
+{
+	// Initialize shouldn't be called more than once...
+	if (!ensure(!UnwrapGeometryActor))
+	{
+		UnwrapGeometryActor->Destroy();
+	}
+	if (!ensure(!LivePreviewGeometryActor))
+	{
+		LivePreviewGeometryActor->Destroy();
+	}
+	
+	// Owns most of the unwrap geometry except for the unselected paired edges, since we don't
+	// want those to move if we change the actor transform via SetUnwrapHighlightTransform
+	UnwrapGeometryActor = UnwrapWorld->SpawnActor<APreviewGeometryActor>(
+		FVector::ZeroVector, FRotator(0, 0, 0), FActorSpawnParameters());
+	
+	UnwrapTriangleSet = NewObject<UTriangleSetComponent>(UnwrapGeometryActor);
+	// We are setting the TranslucencySortPriority here to handle the UV editor's use case in 2D
+	// where multiple translucent layers are drawn on top of each other but still need depth sorting.
+	UnwrapTriangleSet->TranslucencySortPriority = FUVEditorUXSettings::SelectionTriangleDepthBias;
+	TriangleSetMaterial = ToolSetupUtil::GetCustomTwoSidedDepthOffsetMaterial(GetParentTool()->GetToolManager(),
+		FUVEditorUXSettings::SelectionTriangleFillColor,
+		FUVEditorUXSettings::SelectionTriangleDepthBias,
+		FUVEditorUXSettings::SelectionTriangleOpacity);
+	UnwrapGeometryActor->SetRootComponent(UnwrapTriangleSet.Get());
+	UnwrapTriangleSet->RegisterComponent();
+
+	UnwrapLineSet = NewObject<ULineSetComponent>(UnwrapGeometryActor);
+	UnwrapLineSet->SetLineMaterial(ToolSetupUtil::GetDefaultLineComponentMaterial(GetParentTool()->GetToolManager(), true));
+	UnwrapLineSet->AttachToComponent(UnwrapTriangleSet.Get(), FAttachmentTransformRules::KeepWorldTransform);
+	UnwrapLineSet->RegisterComponent();
+
+	UnwrapPairedEdgeLineSet = NewObject<ULineSetComponent>(UnwrapGeometryActor);
+	UnwrapPairedEdgeLineSet->SetLineMaterial(ToolSetupUtil::GetDefaultLineComponentMaterial(GetParentTool()->GetToolManager(), true));
+	UnwrapPairedEdgeLineSet->AttachToComponent(UnwrapTriangleSet.Get(), FAttachmentTransformRules::KeepWorldTransform);
+	UnwrapPairedEdgeLineSet->RegisterComponent();
+
+	SewEdgePairingLineSet = NewObject<ULineSetComponent>(UnwrapGeometryActor);
+	SewEdgePairingLineSet->SetLineMaterial(ToolSetupUtil::GetDefaultLineComponentMaterial(
+		GetParentTool()->GetToolManager(), /*bDepthTested*/ true));
+	SewEdgePairingLineSet->AttachToComponent(UnwrapTriangleSet.Get(), FAttachmentTransformRules::KeepWorldTransform);
+	SewEdgePairingLineSet->RegisterComponent();
+	SewEdgePairingLineSet->SetVisibility(bPairedEdgeHighlightsEnabled);
+
+	// The unselected paired edges get their own, stationary, actor.
+	UnwrapStationaryGeometryActor = UnwrapWorld->SpawnActor<APreviewGeometryActor>(
+		FVector::ZeroVector, FRotator(0, 0, 0), FActorSpawnParameters());
+	SewEdgeUnselectedPairingLineSet = NewObject<ULineSetComponent>(UnwrapStationaryGeometryActor);
+	SewEdgeUnselectedPairingLineSet->SetLineMaterial(ToolSetupUtil::GetDefaultLineComponentMaterial(
+		GetParentTool()->GetToolManager(), /*bDepthTested*/ true));
+	UnwrapStationaryGeometryActor->SetRootComponent(SewEdgeUnselectedPairingLineSet.Get());
+	SewEdgeUnselectedPairingLineSet->RegisterComponent();
+	SewEdgeUnselectedPairingLineSet->SetVisibility(bPairedEdgeHighlightsEnabled);
+
+	UnwrapPointSet = NewObject<UPointSetComponent>(UnwrapGeometryActor);
+	UnwrapPointSet->SetPointMaterial(ToolSetupUtil::GetDefaultPointComponentMaterial(GetParentTool()->GetToolManager(), true));
+	UnwrapPointSet->AttachToComponent(UnwrapTriangleSet.Get(), FAttachmentTransformRules::KeepWorldTransform);
+	UnwrapPointSet->RegisterComponent();
+
+	// Owns the highlights in the live preview.
+	LivePreviewGeometryActor = LivePreviewWorld->SpawnActor<APreviewGeometryActor>(
+		FVector::ZeroVector, FRotator(0, 0, 0), FActorSpawnParameters());
+	
+	LivePreviewLineSet = NewObject<ULineSetComponent>(LivePreviewGeometryActor);
+	LivePreviewLineSet->SetLineMaterial(ToolSetupUtil::GetDefaultLineComponentMaterial(
+		GetParentTool()->GetToolManager(), /*bDepthTested*/ true));
+	LivePreviewGeometryActor->SetRootComponent(LivePreviewLineSet.Get());
+	LivePreviewLineSet->RegisterComponent();
+	
+	LivePreviewPointSet = NewObject<UPointSetComponent>(LivePreviewGeometryActor);	
+	LivePreviewPointSet->SetPointMaterial(ToolSetupUtil::GetDefaultPointComponentMaterial(
+		GetParentTool()->GetToolManager(), /*bDepthTested*/ true));
+	LivePreviewPointSet->AttachToComponent(LivePreviewLineSet.Get(), FAttachmentTransformRules::KeepWorldTransform);
+	LivePreviewPointSet->RegisterComponent();
+}
+
+void UUVToolSelectionHighlightMechanic::Shutdown()
+{
+	if (UnwrapGeometryActor)
+	{
+		UnwrapGeometryActor->Destroy();
+		UnwrapGeometryActor = nullptr;
+	}
+	if (UnwrapStationaryGeometryActor)
+	{
+		UnwrapStationaryGeometryActor->Destroy();
+		UnwrapStationaryGeometryActor = nullptr;
+	}
+	if (LivePreviewGeometryActor)
+	{
+		LivePreviewGeometryActor->Destroy();
+		LivePreviewGeometryActor = nullptr;
+	}
+
+	TriangleSetMaterial = nullptr;
+}
+
+void UUVToolSelectionHighlightMechanic::SetIsVisible(bool bUnwrapHighlightVisible, bool bLivePreviewHighlightVisible)
+{
+	if (UnwrapGeometryActor)
+	{
+		UnwrapGeometryActor->SetActorHiddenInGame(!bUnwrapHighlightVisible);
+	}
+	if (UnwrapStationaryGeometryActor)
+	{
+		UnwrapStationaryGeometryActor->SetActorHiddenInGame(!bUnwrapHighlightVisible);
+	}
+	if (LivePreviewGeometryActor)
+	{
+		LivePreviewGeometryActor->SetActorHiddenInGame(!bLivePreviewHighlightVisible);
+	}
+}
+
+void UUVToolSelectionHighlightMechanic::RebuildUnwrapHighlight(
+	const TArray<FUVToolSelection>& Selections, const FTransform& StartTransform, 
+	bool bUsePreviews)
+{
+	if (!ensure(UnwrapGeometryActor))
+	{
+		return;
+	}
+
+	UnwrapTriangleSet->Clear();
+	UnwrapLineSet->Clear();
+	UnwrapPointSet->Clear();
+	SewEdgePairingLineSet->Clear();
+	SewEdgeUnselectedPairingLineSet->Clear();
+	StaticPairedEdgeVidsPerMesh.Reset();
+
+	UnwrapGeometryActor->SetActorTransform(StartTransform);
+
+	for (const FUVToolSelection& Selection : Selections)
+	{
+		if (!ensure(Selection.Target.IsValid() && Selection.Target->IsValid()))
+		{
+			return;
+		}
+
+		const FDynamicMesh3& Mesh = bUsePreviews ? *Selection.Target->UnwrapPreview->PreviewMesh->GetMesh()
+			: *Selection.Target->UnwrapCanonical;
+
+		if (Selection.Type == FUVToolSelection::EType::Triangle)
+		{
+			TRACE_CPUPROFILER_EVENT_SCOPE(UUVToolSelectionHighlightMechanic::AppendUnwrapHighlight_Triangle);
+
+			UnwrapTriangleSet->ReserveTriangles(Selection.SelectedIDs.Num());
+			UnwrapLineSet->ReserveLines(Selection.SelectedIDs.Num() * 3);
+			for (int32 Tid : Selection.SelectedIDs)
+			{
+				if (!ensure(Mesh.IsTriangle(Tid)))
+				{
+					continue;
+				}
+
+				FIndex3i Vids = Mesh.GetTriangle(Tid);
+				FVector Points[3];
+				for (int i = 0; i < 3; ++i)
+				{
+					Points[i] = StartTransform.InverseTransformPosition(Mesh.GetVertex(Vids[i]));
+				}
+				UnwrapTriangleSet->AddTriangle(Points[0], Points[1], Points[2], FVector(0, 0, 1), FUVEditorUXSettings::SelectionTriangleFillColor, TriangleSetMaterial);
+				for (int i = 0; i < 3; ++i)
+				{
+					int NextIndex = (i + 1) % 3;
+					UnwrapLineSet->AddLine(Points[i], Points[NextIndex],
+						FUVEditorUXSettings::SelectionTriangleWireframeColor,
+						FUVEditorUXSettings::SelectionLineThickness,
+						FUVEditorUXSettings::SelectionWireframeDepthBias);
+				}
+			}
+		}
+		else if (Selection.Type == FUVToolSelection::EType::Edge)
+		{
+			TRACE_CPUPROFILER_EVENT_SCOPE(MeshSelectionMechanic_RebuildDrawnElements_Edge);
+
+			StaticPairedEdgeVidsPerMesh.Emplace();
+			StaticPairedEdgeVidsPerMesh.Last().Key = Selection.Target;
+
+			const FDynamicMesh3& AppliedMesh = bUsePreviews ? *Selection.Target->AppliedPreview->PreviewMesh->GetMesh()
+				: *Selection.Target->AppliedCanonical;
+
+			UnwrapLineSet->ReserveLines(Selection.SelectedIDs.Num());
+			for (int32 Eid : Selection.SelectedIDs)
+			{
+				if (!ensure(Mesh.IsEdge(Eid)))
+				{
+					continue;
+				}
+
+				FIndex2i EdgeVids = Mesh.GetEdgeV(Eid);
+				UnwrapLineSet->AddLine(
+					StartTransform.InverseTransformPosition(Mesh.GetVertex(EdgeVids.A)),
+					StartTransform.InverseTransformPosition(Mesh.GetVertex(EdgeVids.B)),
+					FUVEditorUXSettings::SelectionTriangleWireframeColor,
+					FUVEditorUXSettings::SelectionLineThickness,
+					FUVEditorUXSettings::SelectionWireframeDepthBias);
+
+				if (bPairedEdgeHighlightsEnabled)
+				{
+					bool bWouldPreferReverse = false;
+					int32 PairedEid = UUVSeamSewAction::FindSewEdgeOppositePairing(
+						Mesh, AppliedMesh, Selection.Target->UVLayerIndex, Eid, bWouldPreferReverse);
+
+					bool bPairedEdgeIsSelected = Selection.SelectedIDs.Contains(PairedEid);
+
+					if (PairedEid == IndexConstants::InvalidID
+						// When both sides are selected, merge order depends on adjacent tid values, so
+						// deal with the pair starting with the other edge.
+						|| (bPairedEdgeIsSelected && bWouldPreferReverse))
+					{
+						continue;
+					}
+
+					FIndex2i Vids1 = Mesh.GetEdgeV(Eid);
+					SewEdgePairingLineSet->AddLine(
+						StartTransform.InverseTransformPosition(Mesh.GetVertex(Vids1.A)),
+						StartTransform.InverseTransformPosition(Mesh.GetVertex(Vids1.B)),
+						FUVEditorUXSettings::SewSideLeftColor, FUVEditorUXSettings::SewLineHighlightThickness, FUVEditorUXSettings::SewLineDepthOffset);
+
+					// The paired edge may need to go into a separate line set if it is not selected so that it does
+					// not get affected by transformations of the selected highlights in SetUnwrapHighlightTransform
+					FIndex2i Vids2 = Mesh.GetEdgeV(PairedEid);
+					if (bPairedEdgeIsSelected)
+					{
+						SewEdgePairingLineSet->AddLine(
+							StartTransform.InverseTransformPosition(Mesh.GetVertex(Vids2.A)),
+							StartTransform.InverseTransformPosition(Mesh.GetVertex(Vids2.B)),
+							FUVEditorUXSettings::SewSideRightColor, FUVEditorUXSettings::SewLineHighlightThickness, 
+							FUVEditorUXSettings::SewLineDepthOffset);
+					}
+					else
+					{
+						StaticPairedEdgeVidsPerMesh.Last().Value.Add(TPair<int32, int32>(Vids2.A, Vids2.B));
+						SewEdgeUnselectedPairingLineSet->AddLine(
+							Mesh.GetVertex(Vids2.A),
+							Mesh.GetVertex(Vids2.B),
+							FUVEditorUXSettings::SewSideRightColor, FUVEditorUXSettings::SewLineHighlightThickness,
+							FUVEditorUXSettings::SewLineDepthOffset);
+					}
+				}//end if visualizing paired edges
+			}//end for each edge
+		}
+		else if (Selection.Type == FUVToolSelection::EType::Vertex)
+		{
+			TRACE_CPUPROFILER_EVENT_SCOPE(MeshSelectionMechanic_RebuildDrawnElements_Vertex);
+
+			UnwrapPointSet->ReservePoints(Selection.SelectedIDs.Num());
+			for (int32 Vid : Selection.SelectedIDs)
+			{
+				if (!ensure(Mesh.IsVertex(Vid)))
+				{
+					continue;
+				}
+
+				FRenderablePoint PointToRender(StartTransform.InverseTransformPosition(Mesh.GetVertex(Vid)),
+					FUVEditorUXSettings::SelectionTriangleWireframeColor,
+					FUVEditorUXSettings::SelectionPointThickness,
+					FUVEditorUXSettings::SelectionWireframeDepthBias);
+				UnwrapPointSet->AddPoint(PointToRender);
+			}
+		}
+	}
+}
+
+void UUVToolSelectionHighlightMechanic::SetUnwrapHighlightTransform(const FTransform& Transform, 
+	bool bRebuildStaticPairedEdges, bool bUsePreviews)
+{
+	if (ensure(UnwrapGeometryActor))
+	{
+		UnwrapGeometryActor->SetActorTransform(Transform);
+	}
+	if (bPairedEdgeHighlightsEnabled && bRebuildStaticPairedEdges)
+	{
+		SewEdgeUnselectedPairingLineSet->Clear();
+		for (const TPair<TWeakObjectPtr<UUVEditorToolMeshInput>, 
+			TArray<TPair<int32, int32>>>& MeshVidPairs : StaticPairedEdgeVidsPerMesh)
+		{
+			TWeakObjectPtr<UUVEditorToolMeshInput> Target = MeshVidPairs.Key;
+			if (!ensure(Target.IsValid()))
+			{
+				continue;
+			}
+
+			const FDynamicMesh3& Mesh = bUsePreviews ? *Target->UnwrapPreview->PreviewMesh->GetMesh()
+				: *Target->UnwrapCanonical;
+
+			for (const TPair<int32, int32>& VidPair : MeshVidPairs.Value)
+			{
+				if (!ensure(Mesh.IsVertex(VidPair.Key) && Mesh.IsVertex(VidPair.Value)))
+				{
+					continue;
+				}
+				SewEdgeUnselectedPairingLineSet->AddLine(
+					Mesh.GetVertex(VidPair.Key),
+					Mesh.GetVertex(VidPair.Value),
+					FUVEditorUXSettings::SewSideRightColor, FUVEditorUXSettings::SewLineHighlightThickness,
+					FUVEditorUXSettings::SewLineDepthOffset);
+			}
+		}
+	}
+}
+
+FTransform UUVToolSelectionHighlightMechanic::GetUnwrapHighlightTransform()
+{
+	if (ensure(UnwrapGeometryActor))
+	{
+		return UnwrapGeometryActor->GetActorTransform();
+	}
+	return FTransform::Identity;
+}
+
+void UUVToolSelectionHighlightMechanic::RebuildAppliedHighlightFromUnwrapSelection(
+	const TArray<FUVToolSelection>& UnwrapSelections, bool bUsePreviews)
+{
+	if (!ensure(LivePreviewGeometryActor))
+	{
+		return;
+	}
+
+	LivePreviewLineSet->Clear();
+	LivePreviewPointSet->Clear();
+
+	for (const FUVToolSelection& Selection : UnwrapSelections)
+	{
+		if (!ensure(Selection.Target.IsValid() && Selection.Target->IsValid()))
+		{
+			return;
+		}
+
+		UUVEditorToolMeshInput* Target = Selection.Target.Get();
+
+		const FDynamicMesh3& AppliedMesh = bUsePreviews ? *Target->AppliedPreview->PreviewMesh->GetMesh()
+			: *Target->AppliedCanonical;
+		const FDynamicMesh3& UnwrapMesh = bUsePreviews ? *Target->UnwrapPreview->PreviewMesh->GetMesh()
+			: *Target->UnwrapCanonical;
+
+		FTransform MeshTransform = Target->AppliedPreview->PreviewMesh->GetTransform();
+
+		auto AppendEdgeLine = [this, &AppliedMesh, &MeshTransform](int32 AppliedEid)
+		{
+			FVector3d Vert1, Vert2;
+			AppliedMesh.GetEdgeV(AppliedEid, Vert1, Vert2);
+
+			LivePreviewLineSet->AddLine(
+				MeshTransform.TransformPosition(Vert1),
+				MeshTransform.TransformPosition(Vert2),
+				FUVEditorUXSettings::SelectionTriangleWireframeColor,
+				FUVEditorUXSettings::LivePreviewHighlightThickness,
+				FUVEditorUXSettings::LivePreviewHighlightDepthOffset);
+		};
+
+		if (Selection.Type == FUVToolSelection::EType::Triangle)
+		{
+			TRACE_CPUPROFILER_EVENT_SCOPE(Triangle);
+
+			for (int32 Tid : Selection.SelectedIDs)
+			{
+				if (!ensure(AppliedMesh.IsTriangle(Tid)))
+				{
+					continue;
+				}
+
+				// Gather the boundary edges for the live preview
+				FIndex3i TriEids = AppliedMesh.GetTriEdges(Tid);
+				for (int i = 0; i < 3; ++i)
+				{
+					FIndex2i EdgeTids = AppliedMesh.GetEdgeT(TriEids[i]);
+					for (int j = 0; j < 2; ++j)
+					{
+						if (EdgeTids[j] != Tid && !Selection.SelectedIDs.Contains(EdgeTids[j]))
+						{
+							AppendEdgeLine(TriEids[i]);
+							break;
+						}
+					}
+				}//end for tri edges
+			}//end for selection tids
+		}
+		else if (Selection.Type == FUVToolSelection::EType::Edge)
+		{
+			TRACE_CPUPROFILER_EVENT_SCOPE(Edge);
+
+			for (int32 UnwrapEid : Selection.SelectedIDs)
+			{
+				if (!ensure(UnwrapMesh.IsEdge(UnwrapEid)))
+				{
+					continue;
+				}
+
+				FDynamicMesh3::FEdge Edge = UnwrapMesh.GetEdge(UnwrapEid);
+
+				int32 AppliedEid = AppliedMesh.FindEdgeFromTri(
+					Target->UnwrapVidToAppliedVid(Edge.Vert.A),
+					Target->UnwrapVidToAppliedVid(Edge.Vert.B),
+					Edge.Tri.A);
+
+				AppendEdgeLine(AppliedEid);
+			}
+		}
+		else if (Selection.Type == FUVToolSelection::EType::Vertex)
+		{
+			TRACE_CPUPROFILER_EVENT_SCOPE(Vertex);
+
+			for (int32 UnwrapVid : Selection.SelectedIDs)
+			{
+				FVector3d Position = AppliedMesh.GetVertex(
+					Target->UnwrapVidToAppliedVid(UnwrapVid));
+
+				LivePreviewPointSet->AddPoint(Position,
+					FUVEditorUXSettings::SelectionTriangleWireframeColor,
+					FUVEditorUXSettings::LivePreviewHighlightPointSize,
+					FUVEditorUXSettings::LivePreviewHighlightDepthOffset);
+			}
+		}
+	}//end for selection
+}
+
+void UUVToolSelectionHighlightMechanic::SetEnablePairedEdgeHighlights(bool bEnable)
+{
+	bPairedEdgeHighlightsEnabled = bEnable;
+	SewEdgePairingLineSet->SetVisibility(bEnable);
+	SewEdgeUnselectedPairingLineSet->SetVisibility(bEnable);
+}
+
+#undef LOCTEXT_NAMESPACE
\ No newline at end of file
diff --git a/Engine/Plugins/Experimental/UVEditor/Source/UVEditorTools/Private/UVEditorLayerEditTool.cpp b/Engine/Plugins/Experimental/UVEditor/Source/UVEditorTools/Private/UVEditorLayerEditTool.cpp
--- a/Engine/Plugins/Experimental/UVEditor/Source/UVEditorTools/Private/UVEditorLayerEditTool.cpp
+++ b/Engine/Plugins/Experimental/UVEditor/Source/UVEditorTools/Private/UVEditorLayerEditTool.cpp
@@ -11,7 +11,7 @@
 #include "ParameterizationOps/UVLayoutOp.h"
 #include "Properties/UVLayoutProperties.h"
 #include "ToolTargets/UVEditorToolMeshInput.h"
-#include "UVToolContextObjects.h"
+#include "ContextObjects/UVToolContextObjects.h"
 #include "ModelingToolTargetUtil.h"
 #include "Components.h"
 #include "UVEditorToolAnalyticsUtils.h"
diff --git a/Engine/Plugins/Experimental/UVEditor/Source/UVEditorTools/Private/UVEditorLayoutTool.cpp b/Engine/Plugins/Experimental/UVEditor/Source/UVEditorTools/Private/UVEditorLayoutTool.cpp
--- a/Engine/Plugins/Experimental/UVEditor/Source/UVEditorTools/Private/UVEditorLayoutTool.cpp
+++ b/Engine/Plugins/Experimental/UVEditor/Source/UVEditorTools/Private/UVEditorLayoutTool.cpp
@@ -10,7 +10,7 @@
 #include "ParameterizationOps/UVLayoutOp.h"
 #include "Properties/UVLayoutProperties.h"
 #include "ToolTargets/UVEditorToolMeshInput.h"
-#include "UVToolContextObjects.h"
+#include "ContextObjects/UVToolContextObjects.h"
 #include "EngineAnalytics.h"
 #include "UVEditorToolAnalyticsUtils.h"
 
@@ -68,7 +68,7 @@ void UUVEditorLayoutTool::Setup()
 	SetToolDisplayName(LOCTEXT("ToolName", "UV Layout"));
 	GetToolManager()->DisplayMessage(LOCTEXT("OnStartUVLayoutTool", "Translate, rotate or scale existing UV Charts using various strategies"),
 		EToolMessageLevel::UserNotification);
-
+	
 	// Analytics
 	InputTargetAnalytics = UVEditorAnalytics::CollectTargetAnalytics(Targets);
 }
diff --git a/Engine/Plugins/Experimental/UVEditor/Source/UVEditorTools/Private/UVEditorRecomputeUVsTool.cpp b/Engine/Plugins/Experimental/UVEditor/Source/UVEditorTools/Private/UVEditorRecomputeUVsTool.cpp
--- a/Engine/Plugins/Experimental/UVEditor/Source/UVEditorTools/Private/UVEditorRecomputeUVsTool.cpp
+++ b/Engine/Plugins/Experimental/UVEditor/Source/UVEditorTools/Private/UVEditorRecomputeUVsTool.cpp
@@ -7,7 +7,7 @@
 #include "Polygroups/PolygroupUtil.h"
 #include "ParameterizationOps/RecomputeUVsOp.h"
 #include "ToolTargets/UVEditorToolMeshInput.h"
-#include "UVToolContextObjects.h"
+#include "ContextObjects/UVToolContextObjects.h"
 #include "ContextObjectStore.h"
 #include "MeshOpPreviewHelpers.h"
 #include "UVEditorToolAnalyticsUtils.h"
diff --git a/Engine/Plugins/Experimental/UVEditor/Source/UVEditorTools/Private/UVEditorSeamTool.cpp b/Engine/Plugins/Experimental/UVEditor/Source/UVEditorTools/Private/UVEditorSeamTool.cpp
--- a/Engine/Plugins/Experimental/UVEditor/Source/UVEditorTools/Private/UVEditorSeamTool.cpp
+++ b/Engine/Plugins/Experimental/UVEditor/Source/UVEditorTools/Private/UVEditorSeamTool.cpp
@@ -16,7 +16,7 @@
 #include "ToolTargets/UVEditorToolMeshInput.h"
 #include "ToolSceneQueriesUtil.h"
 #include "UVEditorUXSettings.h"
-#include "UVToolContextObjects.h"
+#include "ContextObjects/UVToolContextObjects.h"
 #include "EngineAnalytics.h"
 
 #define LOCTEXT_NAMESPACE "UUVEditorSeamTool"
diff --git a/Engine/Plugins/Experimental/UVEditor/Source/UVEditorTools/Private/UVIslandConformalUnwrapAction.cpp b/Engine/Plugins/Experimental/UVEditor/Source/UVEditorTools/Private/UVIslandConformalUnwrapAction.cpp
deleted file mode 100644
diff --git a/Engine/Plugins/Experimental/UVEditor/Source/UVEditorTools/Private/UVSeamSewAction.cpp b/Engine/Plugins/Experimental/UVEditor/Source/UVEditorTools/Private/UVSeamSewAction.cpp
deleted file mode 100644
diff --git a/Engine/Plugins/Experimental/UVEditor/Source/UVEditorTools/Private/UVSelectTool.cpp b/Engine/Plugins/Experimental/UVEditor/Source/UVEditorTools/Private/UVSelectTool.cpp
--- a/Engine/Plugins/Experimental/UVEditor/Source/UVEditorTools/Private/UVSelectTool.cpp
+++ b/Engine/Plugins/Experimental/UVEditor/Source/UVEditorTools/Private/UVSelectTool.cpp
@@ -2,29 +2,19 @@
 
 #include "UVSelectTool.h"
 
-#include "Algo/Unique.h"
 #include "BaseGizmos/CombinedTransformGizmo.h"
 #include "BaseGizmos/GizmoBaseComponent.h"
+#include "ContextObjects/UVToolViewportButtonsAPI.h"
 #include "ContextObjectStore.h"
-#include "Drawing/LineSetComponent.h"
-#include "Drawing/MeshElementsVisualizer.h"
-#include "Drawing/PointSetComponent.h"
-#include "Drawing/PreviewGeometryActor.h"
 #include "DynamicMesh/DynamicMeshChangeTracker.h"
 #include "DynamicMesh/MeshIndexUtil.h"
-#include "ToolTargets/UVEditorToolMeshInput.h"
 #include "InteractiveToolManager.h"
 #include "MeshOpPreviewHelpers.h" // UMeshOpPreviewWithBackgroundCompute
-#include "Parameterization/DynamicMeshUVEditor.h"
 #include "PreviewMesh.h"
-#include "Selection/UVEditorMeshSelectionMechanic.h"
-#include "Selection/UVEditorDynamicMeshSelection.h"
+#include "Selection/UVToolSelection.h"
 #include "ToolSetupUtil.h"
+#include "ToolTargets/UVEditorToolMeshInput.h"
 #include "UVEditorUXSettings.h"
-#include "EngineAnalytics.h"
-
-#include "UVSeamSewAction.h"
-#include "UVIslandConformalUnwrapAction.h"
 
 #include "ToolTargetManager.h"
 
@@ -34,143 +24,62 @@ using namespace UE::Geometry;
 
 namespace UVSelectToolLocals
 {
-	// These following three functions deal with the unfortunate problem that eids are unstable as identifiers
-	// (e.g. removing and reinserting the same triangles can change the eids of the edges), so edges have to
-	// be identified in another way. We identify them by vertex ID pairs. This should really be dealt with
-	// on a mesh selection level, but for now we fix it here.
-	// After selection changes, we convert our eids to vid pairs. After mesh changes, we update the selection
-	// eids from our stored vid pairs.
-
-	/** If selection is a non-empty edge selection, update its eids using stored vid pairs. */
-	void UpdateSelectionEidsAfterMeshChange(FUVEditorDynamicMeshSelection& SelectionInOut, TArray<FIndex2i>* VidPairsIn)
-	{
-		if (!SelectionInOut.Mesh || SelectionInOut.Type != FUVEditorDynamicMeshSelection::EType::Edge)
-		{
-			// No update necessary
-			return;
-		}
-
-		// Otherwise, updating eids.
-		if (!ensure(VidPairsIn))
-		{
-			return;
-		}
-		SelectionInOut.SelectedIDs.Empty();
-		for (const FIndex2i& VidPair : *VidPairsIn)
-		{
-			int32 Eid = SelectionInOut.Mesh->FindEdge(VidPair.A, VidPair.B);
-			if (ensure(Eid != IndexConstants::InvalidID))
-			{
-				SelectionInOut.SelectedIDs.Add(Eid);
-			}
-		}
-	}
-
-	/** If selection mechanic holds a non-empty edge selection, update its eids using stored vid pairs. */
-	void UpdateSelectionEidsAfterMeshChange(UUVEditorMeshSelectionMechanic& SelectionMechanic, TArray<FIndex2i>* VidPairsIn)
-	{
-		const FUVEditorDynamicMeshSelection& CurrentSelection = SelectionMechanic.GetCurrentSelection();
-		if (CurrentSelection.Mesh && CurrentSelection.Type == FUVEditorDynamicMeshSelection::EType::Edge)
-		{
-			FUVEditorDynamicMeshSelection UpdatedSelection = CurrentSelection;
-			UpdateSelectionEidsAfterMeshChange(UpdatedSelection, VidPairsIn);
-			SelectionMechanic.SetSelection(UpdatedSelection, false, false);
-		}
-	}
-
-	void GetVidPairsFromSelection(const FUVEditorDynamicMeshSelection& SelectionIn, TArray<FIndex2i>& VidPairsOut)
-	{
-		VidPairsOut.Reset();
-		if (!SelectionIn.Mesh || SelectionIn.Type != FUVEditorDynamicMeshSelection::EType::Edge)
-		{
-			// No vid pairs to add
-			return;
-		}
-
-		// Otherwise create the vid pairs out of eids
-		for (int32 Eid : SelectionIn.SelectedIDs)
-		{
-			VidPairsOut.Add(SelectionIn.Mesh->GetEdgeV(Eid));
-		}
-	}
+	FText GizmoChangeTransactionName = LOCTEXT("GizmoChange", "Gizmo Update");
 
 	/**
-	 * An undo/redo object for selection changes that, instead of operating directly on a selection
-	 * mechanic, instead operates on a context object that tools can use to route the request
-	 * to the current selection mechanic. This is valuable because we want the selection changes
-	 * to be undoable in different invocations of the tool, and the selection mechanic pointer
-	 * will not stay the same. However, the context object will stay the same, and we can register
-	 * to its delegate on each invocation.
-	 * 
-	 * The other thing that is different about this selection change object is that in cases of edge
-	 * selections, it uses stored vid pairs rather then eids, to deal with mesh changes that alter eids.
+	 * Helper change that allows us to preserve the rotational component of the gizmo, since
+	 * just updating the gizmo from the current selection will always reset rotation.
 	 */
-	class FSelectionChange : public FToolCommandChange
+	class  FGizmoChange : public FToolCommandChange
 	{
 	public:
-		/**
-		 * @param GizmoBeforeIn Transform to which to revert gizmo on the way back (to avoid losing gizmo rotation, which
-		 *   gets transacted post-selection-change, and therefore in the wrong place for undo).
-		 * @param EdgeVidPairsBeforeIn
-		 * @param EdgeVidPairsAfterIn
-		 */
-		FSelectionChange(const FUVEditorDynamicMeshSelection& SelectionBeforeIn,
-			const FUVEditorDynamicMeshSelection& SelectionAfterIn,
-			const FTransform& GizmoBeforeIn,
-			TUniquePtr<TArray<FIndex2i>> EdgeVidPairsBeforeIn,
-			TUniquePtr<TArray<FIndex2i>> EdgeVidPairsAfterIn
-			)
-			: SelectionBefore(SelectionBeforeIn)
-			, SelectionAfter(SelectionAfterIn)
-			, GizmoBefore(GizmoBeforeIn)
-			, EdgeVidPairsBefore(MoveTemp(EdgeVidPairsBeforeIn))
-			, EdgeVidPairsAfter(MoveTemp(EdgeVidPairsAfterIn))
+		FGizmoChange(const FTransform& GizmoBeforeIn)
+			: GizmoBefore(GizmoBeforeIn)
 		{
-			// Make sure that for both selections, if we have a non-empty edge selection, we have vid pairs.
-			ensure(!(
-				(SelectionBefore.Mesh && SelectionBefore.Type == FUVEditorDynamicMeshSelection::EType::Edge && !EdgeVidPairsBefore)
-				|| (SelectionAfter.Mesh && SelectionAfter.Type == FUVEditorDynamicMeshSelection::EType::Edge && !EdgeVidPairsAfter)));
-		}
+		};
 
 		virtual void Apply(UObject* Object) override
 		{
-			UUVSelectToolChangeRouter* ChangeRouter = Cast<UUVSelectToolChangeRouter>(Object);
-			if (ensure(ChangeRouter) && ChangeRouter->CurrentSelectTool.IsValid())
+			// Route the action to the currently running instance of the select tool
+			UInteractiveToolManager* ToolManager = Cast<UInteractiveToolManager>(Object);
+			UUVSelectTool* SelectTool = Cast<UUVSelectTool>(ToolManager->GetActiveTool(EToolSide::Mouse));
+			if (ensure(SelectTool))
 			{
-				UpdateSelectionEidsAfterMeshChange(SelectionAfter, EdgeVidPairsAfter.Get());
-				ChangeRouter->CurrentSelectTool->SetSelection(SelectionAfter);
+				SelectTool->SetGizmoTransform(GizmoAfter);
 			}
 		}
 
 		virtual void Revert(UObject* Object) override
 		{
-			UUVSelectToolChangeRouter* ChangeRouter = Cast<UUVSelectToolChangeRouter>(Object);
-			if (ensure(ChangeRouter) && ChangeRouter->CurrentSelectTool.IsValid())
+			// Route the action to the currently running instance of the select tool
+			UInteractiveToolManager* ToolManager = Cast<UInteractiveToolManager>(Object);
+			UUVSelectTool* SelectTool = Cast<UUVSelectTool>(ToolManager->GetActiveTool(EToolSide::Mouse));
+			if (ensure(SelectTool))
 			{
-				UpdateSelectionEidsAfterMeshChange(SelectionBefore, EdgeVidPairsBefore.Get());
-				ChangeRouter->CurrentSelectTool->SetSelection(SelectionBefore);
-				ChangeRouter->CurrentSelectTool->SetGizmoTransform(GizmoBefore);
+				// This is the easiest way to make sure the transform reverts back to whatever it was changed to
+				// on redo (which we may not fully know at time of emitting if we're emitting on tool shutdown,
+				// and the next transform is whatever exists at next tool invocation).
+				GizmoAfter = SelectTool->GetGizmoTransform();
+				SelectTool->SetGizmoTransform(GizmoBefore);
 			}
 		}
 
 		virtual bool HasExpired(UObject* Object) const override
 		{
-			UUVSelectToolChangeRouter* ChangeRouter = Cast<UUVSelectToolChangeRouter>(Object);
-			return !(ChangeRouter && ChangeRouter->CurrentSelectTool.IsValid());
+			// Expired if currently active tool cannot be cast to select tool
+			UInteractiveToolManager* ToolManager = Cast<UInteractiveToolManager>(Object);
+			return Cast<UUVSelectTool>(ToolManager->GetActiveTool(EToolSide::Mouse)) == nullptr;
 		}
 
+
 		virtual FString ToString() const override
 		{
-			return TEXT("UVSelectToolLocals::FSelectionChange");
+			return TEXT("UVSelectToolLocals::FGizmoChange");
 		}
 
 	protected:
-		FUVEditorDynamicMeshSelection SelectionBefore;
-		FUVEditorDynamicMeshSelection SelectionAfter;
 		FTransform GizmoBefore;
-
-		TUniquePtr<TArray<FIndex2i>> EdgeVidPairsBefore;
-		TUniquePtr<TArray<FIndex2i>> EdgeVidPairsAfter;
+		FTransform GizmoAfter;
 	};
 
 	/**
@@ -208,10 +117,11 @@ namespace UVSelectToolLocals
 			UnwrapCanonicalMeshChange->Apply(UVToolInputObject->UnwrapCanonical.Get(), false);
 			UVToolInputObject->UpdateFromCanonicalUnwrapUsingMeshChange(*UnwrapCanonicalMeshChange);
 
-			UUVSelectToolChangeRouter* ChangeRouter = Cast<UUVSelectToolChangeRouter>(Object);
-			if (ensure(ChangeRouter) && ChangeRouter->CurrentSelectTool.IsValid())
+			UInteractiveToolManager* ToolManager = Cast<UInteractiveToolManager>(Object);
+			UUVSelectTool* SelectTool = Cast<UUVSelectTool>(ToolManager->GetActiveTool(EToolSide::Mouse));
+			if (SelectTool)
 			{
-				ChangeRouter->CurrentSelectTool->SetGizmoTransform(GizmoAfter);
+				SelectTool->SetGizmoTransform(GizmoAfter);
 			}
 
 		}
@@ -221,10 +131,11 @@ namespace UVSelectToolLocals
 			UnwrapCanonicalMeshChange->Apply(UVToolInputObject->UnwrapCanonical.Get(), true);
 			UVToolInputObject->UpdateFromCanonicalUnwrapUsingMeshChange(*UnwrapCanonicalMeshChange);
 
-			UUVSelectToolChangeRouter* ChangeRouter = Cast<UUVSelectToolChangeRouter>(Object);
-			if (ensure(ChangeRouter) && ChangeRouter->CurrentSelectTool.IsValid())
+			UInteractiveToolManager* ToolManager = Cast<UInteractiveToolManager>(Object);
+			UUVSelectTool* SelectTool = Cast<UUVSelectTool>(ToolManager->GetActiveTool(EToolSide::Mouse));
+			if (SelectTool)
 			{
-				ChangeRouter->CurrentSelectTool->SetGizmoTransform(GizmoBefore);
+				SelectTool->SetGizmoTransform(GizmoBefore);
 			}
 		}
 
@@ -261,145 +172,58 @@ bool UUVSelectToolBuilder::CanBuildTool(const FToolBuilderState& SceneState) con
 UInteractiveTool* UUVSelectToolBuilder::BuildTool(const FToolBuilderState& SceneState) const
 {
 	UUVSelectTool* NewTool = NewObject<UUVSelectTool>(SceneState.ToolManager);
-	NewTool->SetWorld(SceneState.World);
 	NewTool->SetTargets(*Targets);
 
 	return NewTool;
 }
 
-// Tool property functions
-void  USelectToolActionPropertySet::IslandConformalUnwrap()
-{
-	PostAction(ESelectToolAction::IslandConformalUnwrap);
-}
-
-void USelectToolActionPropertySet::PostAction(ESelectToolAction Action)
-{
-	if (ParentTool.IsValid())
-	{
-		ParentTool->RequestAction(Action);
-	}
-}
-
-
 void UUVSelectTool::Setup()
 {
 	TRACE_CPUPROFILER_EVENT_SCOPE(UVSelectTool_Setup);
-	
-	ToolStartTimeAnalytics = FDateTime::UtcNow();
 
 	using namespace UVSelectToolLocals;
 
 	check(Targets.Num() > 0);
 
 	UInteractiveTool::Setup();
-	
+
 	SetToolDisplayName(LOCTEXT("ToolName", "UV Select Tool"));
 
+	// Get all the API's we'll need
 	UContextObjectStore* ContextStore = GetToolManager()->GetContextObjectStore();
 	EmitChangeAPI = ContextStore->FindContext<UUVToolEmitChangeAPI>();
 	ViewportButtonsAPI = ContextStore->FindContext<UUVToolViewportButtonsAPI>();
+	SelectionAPI = ContextStore->FindContext<UUVToolSelectionAPI>();
+
 	ViewportButtonsAPI->SetGizmoButtonsEnabled(true);
 	ViewportButtonsAPI->OnGizmoModeChange.AddWeakLambda(this, 
 		[this](UUVToolViewportButtonsAPI::EGizmoMode NewGizmoMode) {
 			UpdateGizmo();
 		});
-	ViewportButtonsAPI->SetSelectionButtonsEnabled(true);
-	ViewportButtonsAPI->OnSelectionModeChange.AddWeakLambda(this,
-		[this](UUVToolViewportButtonsAPI::ESelectionMode NewMode) {
-			UpdateSelectionMode();
-		});
-
-	ToolActions = NewObject<USelectToolActionPropertySet>(this);
-	ToolActions->Initialize(this);
-	AddToolPropertySource(ToolActions);
-
-	SelectionMechanic = NewObject<UUVEditorMeshSelectionMechanic>();
-	SelectionMechanic->Setup(this);
-	SelectionMechanic->SetWorld(Targets[0]->UnwrapPreview->GetWorld());
-	SelectionMechanic->OnSelectionChanged.AddUObject(this, &UUVSelectTool::OnSelectionChanged);
-		
-	// Make it so that our selection mechanic creates undo/redo transactions that go to a selection
-	// change router, which we use to route to the current selection mechanic on each tool invocation.
-	ChangeRouter = ContextStore->FindContext<UUVSelectToolChangeRouter>();
-	if (!ChangeRouter)
-	{
-		ChangeRouter = NewObject<UUVSelectToolChangeRouter>();
-		ContextStore->AddContextObject(ChangeRouter);
-	}
-	ChangeRouter->CurrentSelectTool = this;
-
-	SelectionMechanic->EmitSelectionChange = [this](const FUVEditorDynamicMeshSelection& OldSelection,
-		const FUVEditorDynamicMeshSelection& NewSelection)
-	{
-		TUniquePtr<TArray<FIndex2i>> VidPairsBefore;
-		TUniquePtr<TArray<FIndex2i>> VidPairsAfter;
-		if (OldSelection.Type == FUVEditorDynamicMeshSelection::EType::Edge)
-		{
-			VidPairsBefore = MakeUnique<TArray<FIndex2i>>();
-			GetVidPairsFromSelection(OldSelection, *VidPairsBefore);
-		}
-		if (NewSelection.Type == FUVEditorDynamicMeshSelection::EType::Edge)
-		{
-			VidPairsAfter = MakeUnique<TArray<FIndex2i>>();
-			GetVidPairsFromSelection(NewSelection, *VidPairsAfter);
-		}
-		EmitChangeAPI->EmitToolIndependentChange(ChangeRouter, MakeUnique<UVSelectToolLocals::FSelectionChange>(
-			OldSelection, NewSelection, TransformGizmo->GetGizmoTransform(),
-			MoveTemp(VidPairsBefore), MoveTemp(VidPairsAfter)),
-			LOCTEXT("SelectionChangeMessage", "Selection Change"));
-	};
-
-	UpdateSelectionMode();
-
-	// Retrieve cached AABB tree storage, or else set it up
-	UUVToolAABBTreeStorage* TreeStore = ContextStore->FindContext<UUVToolAABBTreeStorage>();
-	if (!TreeStore)
-	{
-		TreeStore = NewObject<UUVToolAABBTreeStorage>();
-		ContextStore->AddContextObject(TreeStore);
-	}
-
-	// Initialize the AABB trees from cached values, or make new ones.
-	for (TObjectPtr<UUVEditorToolMeshInput> Target : Targets)
-	{
-		TSharedPtr<FDynamicMeshAABBTree3> Tree = TreeStore->Get(Target->UnwrapCanonical.Get());
-		if (!Tree)
-		{
-			TRACE_CPUPROFILER_EVENT_SCOPE(BuildAABBTreeForTarget);
-			Tree = MakeShared<FDynamicMeshAABBTree3>();
-			Tree->SetMesh(Target->UnwrapCanonical.Get(), false);
-			// For now we split round-robin on the X/Y axes TODO Experiment with better splitting heuristics
-			FDynamicMeshAABBTree3::GetSplitAxisFunc GetSplitAxis = [](int Depth, const FAxisAlignedBox3d&) { return Depth % 2; };
-			// Note: 16 tris/leaf was chosen with data collected by SpatialBenchmarks.cpp in GeometryProcessingUnitTests
-			Tree->SetBuildOptions(16, MoveTemp(GetSplitAxis));
-			Tree->Build();
-			TreeStore->Set(Target->UnwrapCanonical.Get(), Tree, Target);
-		}
-		AABBTrees.Add(Tree);
-	}
 
-	// Add the spatial structures to the selection mechanic
-	for (int32 i = 0; i < Targets.Num(); ++i)
-	{
-		SelectionMechanic->AddSpatial(AABBTrees[i],
-			Targets[i]->UnwrapPreview->PreviewMesh->GetTransform());
-	}
+	SelectionAPI->OnSelectionChanged.AddUObject(this, &UUVSelectTool::OnSelectionChanged);
+	SelectionAPI->OnPreSelectionChange.AddWeakLambda(this, 
+		[this](bool bEmitChange) {
+			if (bEmitChange)
+			{
+				EmitChangeAPI->EmitToolIndependentChange(GetToolManager(),
+					MakeUnique<FGizmoChange>(TransformGizmo->GetGizmoTransform()),
+					GizmoChangeTransactionName);
+			}
+		});
 
 	// Make sure that if undo/redo events act on the meshes, we update our state.
 	// The trees will be updated by the tree store, which listens to the same broadcasts.
 	for (int32 i = 0; i < Targets.Num(); ++i)
 	{
 		Targets[i]->OnCanonicalModified.AddWeakLambda(this, [this]
-		(UUVEditorToolMeshInput* InputObject, const UUVEditorToolMeshInput::FCanonicalModifiedInfo& Info) {
-			if (bIgnoreOnCanonicalChange) // Used to avoid reacting to broadcasts that we ourselves caused
+			(UUVEditorToolMeshInput* InputObject, const UUVEditorToolMeshInput::FCanonicalModifiedInfo& Info) 
+		{
+			if (bUpdateGizmoOnCanonicalChange) // Used to avoid reacting to broadcasts that we ourselves caused via gizmo movement
 			{
-				return;
+				UpdateGizmo();
+				SelectionAPI->RebuildUnwrapHighlight(TransformGizmo->GetGizmoTransform());
 			}
-			UpdateSelectionEidsAfterMeshChange(*SelectionMechanic, &CurrentSelectionVidPairs);
-			UpdateGizmo();
-			SelectionMechanic->RebuildDrawnElements(TransformGizmo->GetGizmoTransform());
-			SewAction->UpdateVisualizations();
 		});
 	}
 
@@ -418,8 +242,7 @@ void UUVSelectTool::Setup()
 	// Always align gizmo to x and y axes
 	TransformGizmo->bUseContextCoordinateSystem = false;
 	TransformGizmo->SetActiveTarget(TransformProxy, GetToolManager());
-	TransformGizmo->SetVisibility(ViewportButtonsAPI->GetGizmoMode() != UUVToolViewportButtonsAPI::EGizmoMode::Select);
-
+	
 	// Tell the gizmo to be drawn on top even over translucent-mode materials.
 	// Note: this may someday not be necessary, if we get this to work properly by default. Normally we can't
 	// use this approach in modeling mode because it adds dithering to the occluded sections, but we are able
@@ -433,146 +256,80 @@ void UUVSelectTool::Setup()
 		}
 	}
 
-	LivePreviewGeometryActor = Targets[0]->AppliedPreview->GetWorld()->SpawnActor<APreviewGeometryActor>(
-		FVector::ZeroVector, FRotator(0, 0, 0), FActorSpawnParameters());
-	LivePreviewLineSet = NewObject<ULineSetComponent>(LivePreviewGeometryActor);
-	LivePreviewGeometryActor->SetRootComponent(LivePreviewLineSet);
-	LivePreviewLineSet->RegisterComponent();
-	LivePreviewLineSet->SetLineMaterial(ToolSetupUtil::GetDefaultLineComponentMaterial(
-		GetToolManager(), /*bDepthTested*/ true));
-
-	LivePreviewPointSet = NewObject<UPointSetComponent>(LivePreviewGeometryActor);
-	LivePreviewPointSet->AttachToComponent(LivePreviewLineSet, FAttachmentTransformRules::KeepWorldTransform);
-	LivePreviewPointSet->RegisterComponent();
-	LivePreviewPointSet->SetPointMaterial(ToolSetupUtil::GetDefaultPointComponentMaterial(
-		GetToolManager(), /*bDepthTested*/ true));
-
-	SewAction = NewObject<UUVSeamSewAction>();
-	SewAction->Setup(this);
-	SewAction->SetTargets(Targets);
-	SewAction->SetWorld(Targets[0]->UnwrapPreview->GetWorld());
-
-	IslandConformalUnwrapAction = NewObject<UUVIslandConformalUnwrapAction>();
-	IslandConformalUnwrapAction->Setup(this);
-	IslandConformalUnwrapAction->SetTargets(Targets);
-	IslandConformalUnwrapAction->SetWorld(Targets[0]->UnwrapPreview->GetWorld());
-
-	if (!SelectionMechanic->GetCurrentSelection().IsEmpty())
+	SelectionAPI->SetSelectionMechanicEnabled(true);
+
+	// Turn on auto highlighting for live preview.
+	SelectionAPI->SetHighlightVisible(true, true);
+	UUVToolSelectionAPI::FHighlightOptions HighlightOptions;
+	HighlightOptions.bBaseHighlightOnPreviews = true; // Our translation happens in preview
+	HighlightOptions.bAutoUpdateApplied = true;
+	HighlightOptions.bAutoUpdateUnwrap = true;
+	HighlightOptions.bUseCentroidForUnwrapAutoUpdate = true;
+	HighlightOptions.bShowPairedEdgeHighlights = true;
+	SelectionAPI->SetHighlightOptions(HighlightOptions);
+
+	if (SelectionAPI->HaveSelections())
 	{
-		OnSelectionChanged();
+		// This will also update the gizmo
+		OnSelectionChanged(false);
+		SelectionAPI->RebuildUnwrapHighlight(TransformGizmo->GetGizmoTransform());
+		SelectionAPI->RebuildAppliedPreviewHighlight();
+	}
+	else
+	{
+		// Make sure gizmo is hidden
+		UpdateGizmo();
 	}
-	UpdateGizmo();
 
 	GetToolManager()->DisplayMessage(LOCTEXT("SelectToolStatusBarMessage", 
-		"Select elements in the viewport and then transform them or use one of the action buttons."), 
+		"Select elements in the viewport. Activate transform gizmo in viewport to transform them."), 
 		EToolMessageLevel::UserNotification);
-
-	// Analytics
-	InputTargetAnalytics = UVEditorAnalytics::CollectTargetAnalytics(Targets);
 }
 
 void UUVSelectTool::Shutdown(EToolShutdownType ShutdownType)
 {
+	using namespace UVSelectToolLocals;
+
 	TRACE_CPUPROFILER_EVENT_SCOPE(UVSelectTool_Shutdown);
 	
-	// Clear selection so that it can be restored after undoing back into the select tool
-	if (!SelectionMechanic->GetCurrentSelection().IsEmpty())
+	// If the gizmo has a rotated transform, emit an undo transaction so we can get it back if we
+	// undo back into this tool.
+	if (SelectionAPI->HaveSelections() && !TransformGizmo->GetGizmoTransform().RotationEquals(FTransform::Identity))
 	{
-		// (The broadcast here is so that we still broadcast on undo)
-		SelectionMechanic->SetSelection(UUVEditorMeshSelectionMechanic::FUVEditorDynamicMeshSelection(), true, true);
+		EmitChangeAPI->EmitToolIndependentChange(GetToolManager(),
+			MakeUnique<FGizmoChange>(TransformGizmo->GetGizmoTransform()),
+			GizmoChangeTransactionName);
 	}
 
-	ChangeRouter->CurrentSelectTool = nullptr;
-
 	for (TObjectPtr<UUVEditorToolMeshInput> Target : Targets)
 	{
 		Target->OnCanonicalModified.RemoveAll(this);
 	}
 
-	SelectionMechanic->Shutdown();
-
-	if (LivePreviewGeometryActor)
-	{
-		LivePreviewGeometryActor->Destroy();
-		LivePreviewGeometryActor = nullptr;
-		LivePreviewPointSet = nullptr;
-		LivePreviewLineSet = nullptr;
-	}
-
-	if (SewAction)
-	{
-		SewAction->Shutdown();
-	}
-
-	if (IslandConformalUnwrapAction)
-	{
-		IslandConformalUnwrapAction->Shutdown();
-	}
+	// Don't actually need to do this since selection api should reset after us, but still.
+	SelectionAPI->SetSelectionMechanicEnabled(false);
+	SelectionAPI->SetHighlightVisible(false, false);
 
 	// Calls shutdown on gizmo and destroys it.
 	GetToolManager()->GetPairedGizmoManager()->DestroyAllGizmosByOwner(this);
 
 	ViewportButtonsAPI->OnGizmoModeChange.RemoveAll(this);
-	ViewportButtonsAPI->OnSelectionModeChange.RemoveAll(this);
 	ViewportButtonsAPI->SetGizmoButtonsEnabled(false);
-	ViewportButtonsAPI->SetSelectionButtonsEnabled(false);
 
 	ViewportButtonsAPI = nullptr;
 	EmitChangeAPI = nullptr;
-	ChangeRouter = nullptr;
-
-	// Analytics
-	RecordAnalytics();
+	SelectionAPI = nullptr;
 }
 
-void UUVSelectTool::SetSelection(const FUVEditorDynamicMeshSelection& NewSelection)
+FTransform UUVSelectTool::GetGizmoTransform() const
 {
-	SelectionMechanic->SetSelection(NewSelection, true, 
-		false); // Don't emit undo because this function is called from undo
-	
-	// Make sure the current selection mode is compatible with the new selection we received. Don't broadcast
-	// this part because presumably we've already responded to selection change through selection change broadcast.
-	// TODO: there are a couple things that are not ideal about the below. One is that we always change to
-	// triangle mode when we don't know if the triangles came from island or mesh selection mode. Another is
-	// that we change the selection mode in the mechanic directly rather than going through ChangeSelectionMode,
-	// since we don't want to do the conversions/broadcasts that the setter performs. Still, it's not worth
-	// improving this further because the proper solution will probably involve transacting the selection mode
-	// changes, which we'll probably implement while moving selection up to mode level (along with other changes
-	// that would probably stomp anything we do here)
-	UUVToolViewportButtonsAPI::ESelectionMode CurrentMode = ViewportButtonsAPI->GetSelectionMode();
-	switch (NewSelection.Type)
-	{
-	case FUVEditorDynamicMeshSelection::EType::Vertex:
-		if (CurrentMode != UUVToolViewportButtonsAPI::ESelectionMode::Vertex)
-		{
-			ViewportButtonsAPI->SetSelectionMode(UUVToolViewportButtonsAPI::ESelectionMode::Vertex, false);
-			SelectionMechanic->SelectionMode = EUVEditorMeshSelectionMode::Vertex;
-		}
-		break;
-	case FUVEditorDynamicMeshSelection::EType::Edge:
-		if (CurrentMode != UUVToolViewportButtonsAPI::ESelectionMode::Edge)
-		{
-			ViewportButtonsAPI->SetSelectionMode(UUVToolViewportButtonsAPI::ESelectionMode::Edge, false);
-			SelectionMechanic->SelectionMode = EUVEditorMeshSelectionMode::Edge;
-		}
-		break;
-	case FUVEditorDynamicMeshSelection::EType::Triangle:
-		if (CurrentMode != UUVToolViewportButtonsAPI::ESelectionMode::Triangle
-			&& CurrentMode != UUVToolViewportButtonsAPI::ESelectionMode::Island
-			&& CurrentMode != UUVToolViewportButtonsAPI::ESelectionMode::Mesh)
-		{
-			ViewportButtonsAPI->SetSelectionMode(UUVToolViewportButtonsAPI::ESelectionMode::Triangle, false);
-			SelectionMechanic->SelectionMode = EUVEditorMeshSelectionMode::Triangle;
-		}
-		break;
-	}
-
+	return TransformGizmo->GetGizmoTransform();
 }
 
 void UUVSelectTool::SetGizmoTransform(const FTransform& NewTransform)
 {
 	TransformGizmo->ReinitializeGizmoTransform(NewTransform);
-	SelectionMechanic->RebuildDrawnElements(NewTransform);
+	SelectionAPI->RebuildUnwrapHighlight(NewTransform);
 }
 
 void UUVSelectTool::OnPropertyModified(UObject* PropertySet, FProperty* Property)
@@ -581,253 +338,135 @@ void UUVSelectTool::OnPropertyModified(UObject* PropertySet, FProperty* Property
 
 void UUVSelectTool::UpdateGizmo()
 {
-	const FUVEditorDynamicMeshSelection& Selection = SelectionMechanic->GetCurrentSelection();
-
-	if (!Selection.IsEmpty())
+	if (SelectionAPI->HaveSelections())
 	{
-		FVector3d Centroid = SelectionMechanic->GetCurrentSelectionCentroid();
-
-		TransformGizmo->ReinitializeGizmoTransform(FTransform((FVector)Centroid));
+		TransformGizmo->ReinitializeGizmoTransform(FTransform(SelectionAPI->GetUnwrapSelectionCentroid()));
 	}
 
 	TransformGizmo->SetVisibility(
 		ViewportButtonsAPI->GetGizmoMode() != UUVToolViewportButtonsAPI::EGizmoMode::Select
-		&& !SelectionMechanic->GetCurrentSelection().IsEmpty());
+		&& SelectionAPI->HaveSelections());
 }
 
-void UUVSelectTool::UpdateSelectionMode()
-{
-	EUVEditorMeshSelectionMode TargetMode;
-	switch (ViewportButtonsAPI->GetSelectionMode())
-	{
-	case UUVToolViewportButtonsAPI::ESelectionMode::Vertex:
-		TargetMode = EUVEditorMeshSelectionMode::Vertex;
-		break;	
-	case UUVToolViewportButtonsAPI::ESelectionMode::Edge:
-		TargetMode = EUVEditorMeshSelectionMode::Edge;
-		break;
-	case UUVToolViewportButtonsAPI::ESelectionMode::Triangle:
-		TargetMode = EUVEditorMeshSelectionMode::Triangle;
-		break;
-	case UUVToolViewportButtonsAPI::ESelectionMode::Island:
-		TargetMode = EUVEditorMeshSelectionMode::Component;
-		break;
-	case UUVToolViewportButtonsAPI::ESelectionMode::Mesh:
-		TargetMode = EUVEditorMeshSelectionMode::Mesh;
-		break;
-	default:
-		// We shouldn't ever get "none" as the selection mode...
-		ensure(false);
-		TargetMode = EUVEditorMeshSelectionMode::Vertex;
-		break;
-	}
-	SelectionMechanic->ChangeSelectionMode(TargetMode); // broadcast and emit undo if needed
-}
-
-void UUVSelectTool::OnSelectionChanged()
+void UUVSelectTool::OnSelectionChanged(bool)
 {
 	TRACE_CPUPROFILER_EVENT_SCOPE(UVSelectTool_OnSelectionChanged);
 	
 	using namespace UVSelectToolLocals;
 
-	ClearWarning();
-
-	const FUVEditorDynamicMeshSelection& Selection = SelectionMechanic->GetCurrentSelection();
-
-	GetVidPairsFromSelection(Selection, CurrentSelectionVidPairs);
+	MovingVidsPerSelection.Reset();
+	MovingVertOriginalPositionsPerSelection.Reset();
+	RenderUpdateTidsPerSelection.Reset();
 
-	SelectionTargetIndex = -1;
-	MovingVids.Reset();
-	SelectedTids.Reset();
-	LivePreviewEids.Reset();
-	LivePreviewVids.Reset();
+	CurrentSelections = SelectionAPI->GetSelections();
 
-	if (!Selection.IsEmpty())
+	// Do a check for validity of selection so that we are more robust to tools/actions
+	// improperly setting selection.
+	for (const FUVToolSelection& Selection : CurrentSelections)
 	{
-		// Note which mesh we're selecting in.
-		for (int32 i = 0; i < Targets.Num(); ++i)
+		if (!ensure(Selection.Target.IsValid()
+			&& Selection.AreElementsPresentInMesh(*Selection.Target->UnwrapCanonical)))
 		{
-			if (Targets[i]->UnwrapCanonical.Get() == Selection.Mesh)
-			{
-				SelectionTargetIndex = i;
-				break;
-			}
+			// One way we could end up here is if an action that changes the topology doesn't
+			// properly update the selection. While the mode may clear it for us, undoing that
+			// change will still end up with a selection that is incompatible with the current
+			// topology.
+			SelectionAPI->ClearSelections(false, false);
+			break;
 		}
-		check(SelectionTargetIndex >= 0);
+	}
 
-		UUVEditorToolMeshInput* Target = Targets[SelectionTargetIndex];
-		const FDynamicMesh3* LivePreviewMesh = Target->AppliedCanonical.Get();
+	for (int32 SelectionIndex = 0; SelectionIndex < CurrentSelections.Num(); ++SelectionIndex)
+	{
+		FUVToolSelection& Selection = CurrentSelections[SelectionIndex];
+		UUVEditorToolMeshInput* Target = Selection.Target.Get();
+		FDynamicMesh3* UnwrapMesh = Target->UnwrapCanonical.Get();
 
 		// Note the selected vids
 		TSet<int32> VidSet;
 		TSet<int32> TidSet;
-		if (Selection.Type == FUVEditorDynamicMeshSelection::EType::Triangle)
+		if (Selection.Type == FUVToolSelection::EType::Triangle)
 		{
 			TRACE_CPUPROFILER_EVENT_SCOPE(Triangle);
 
 			for (int32 Tid : Selection.SelectedIDs)
 			{
-				FIndex3i TriVids = Selection.Mesh->GetTriangle(Tid);
-				for (int i = 0; i < 3; ++i)
-				{
-					if (!VidSet.Contains(TriVids[i]))
-					{
-						VidSet.Add(TriVids[i]);
-						MovingVids.Add(TriVids[i]);
-					}
-				}
-				if (!TidSet.Contains(Tid))
-				{
-					TidSet.Add(Tid);
-					SelectedTids.Add(Tid);
-				}
-
-				// Gather the boundary edges in the live preview
-				FIndex3i TriEids = LivePreviewMesh->GetTriEdges(Tid);
+				FIndex3i TriVids = UnwrapMesh->GetTriangle(Tid);
 				for (int i = 0; i < 3; ++i)
 				{
-					FIndex2i EdgeTids = LivePreviewMesh->GetEdgeT(TriEids[i]);
-					for (int j = 0; j < 2; ++j)
-					{
-						if (EdgeTids[j] != Tid && !Selection.SelectedIDs.Contains(EdgeTids[j]))
-						{
-							LivePreviewEids.Add(TriEids[i]);
-							break;
-						}
-					}
+					VidSet.Add(TriVids[i]);
 				}
+				TidSet.Add(Tid);
 			}
 		}
-		else if (Selection.Type == FUVEditorDynamicMeshSelection::EType::Edge)
+		else if (Selection.Type == FUVToolSelection::EType::Edge)
 		{
 			TRACE_CPUPROFILER_EVENT_SCOPE(Edge);
 
 			for (int32 Eid : Selection.SelectedIDs)
 			{
-				FIndex2i EdgeVids = Selection.Mesh->GetEdgeV(Eid);
+				FIndex2i EdgeVids = UnwrapMesh->GetEdgeV(Eid);
 				for (int i = 0; i < 2; ++i)
 				{
-					if (!VidSet.Contains(EdgeVids[i]))
-					{
-						VidSet.Add(EdgeVids[i]);
-						MovingVids.Add(EdgeVids[i]);
-					}
+					VidSet.Add(EdgeVids[i]);
 
 					TArray<int> TidOneRing;
-					Selection.Mesh->GetVtxTriangles(EdgeVids[i], TidOneRing);
-					for (int32 Tid : TidOneRing)
-					{
-						if (!TidSet.Contains(Tid))
-						{
-							TidSet.Add(Tid);
-							SelectedTids.Add(Tid);
-						}
-					}
+					UnwrapMesh->GetVtxTriangles(EdgeVids[i], TidOneRing);
+					TidSet.Append(TidOneRing);
 				}
-
-				// Add the edge highlight in the live preview
-				LivePreviewEids.Add(LivePreviewMesh->FindEdge(
-					Target->UnwrapVidToAppliedVid(EdgeVids.A),
-					Target->UnwrapVidToAppliedVid(EdgeVids.B)));
 			}
 		}
-		else if (Selection.Type == FUVEditorDynamicMeshSelection::EType::Vertex)
+		else if (Selection.Type == FUVToolSelection::EType::Vertex)
 		{
 			TRACE_CPUPROFILER_EVENT_SCOPE(Vertex);
 
 			for (int32 Vid : Selection.SelectedIDs)
 			{
-				if (!VidSet.Contains(Vid))
-				{
-					VidSet.Add(Vid);
-					MovingVids.Add(Vid);
-				}
+				VidSet.Add(Vid);
 
 				TArray<int> TidOneRing;
-				Selection.Mesh->GetVtxTriangles(Vid, TidOneRing);
-				for (int32 Tid : TidOneRing)
-				{
-					if (!TidSet.Contains(Tid))
-					{
-						TidSet.Add(Tid);
-						SelectedTids.Add(Tid);
-					}
-				}
-
-				LivePreviewVids.Add(Target->UnwrapVidToAppliedVid(Vid));
+				UnwrapMesh->GetVtxTriangles(Vid, TidOneRing);
+				TidSet.Append(TidOneRing);
 			}
 		}
 		else
 		{
-			check(false);
+			ensure(false);
 		}
-	}
 
-	SewAction->SetSelection(SelectionTargetIndex, &Selection);
-	IslandConformalUnwrapAction->SetSelection(SelectionTargetIndex, &Selection);
+		MovingVidsPerSelection.Emplace(VidSet.Array());
+		RenderUpdateTidsPerSelection.Emplace(TidSet.Array());
+	}
 
-	UpdateLivePreviewLines();
 	UpdateGizmo();
 }
 
-void UUVSelectTool::ClearWarning()
+void UUVSelectTool::GizmoTransformStarted(UTransformProxy* Proxy)
 {
-	GetToolManager()->DisplayMessage(FText(), EToolMessageLevel::UserWarning);
-}
+	bInDrag = true;
 
-void UUVSelectTool::UpdateLivePreviewLines()
-{
-	TRACE_CPUPROFILER_EVENT_SCOPE(UVSelectTool_UpdateLivePreviewLines);
-	
-	LivePreviewLineSet->Clear();
-	LivePreviewPointSet->Clear();
+	InitialGizmoFrame = FFrame3d(TransformGizmo->ActiveTarget->GetTransform());
+
+	MovingVertOriginalPositionsPerSelection.Reset();
 
-	const FUVEditorDynamicMeshSelection& Selection = SelectionMechanic->GetCurrentSelection();
-	if (!Selection.IsEmpty())
+	for (int32 SelectionIndex = 0; SelectionIndex < CurrentSelections.Num(); ++SelectionIndex)
 	{
-		FTransform MeshTransform = Targets[SelectionTargetIndex]->AppliedPreview->PreviewMesh->GetTransform();
-		const FDynamicMesh3* LivePreviewMesh = Targets[SelectionTargetIndex]->AppliedCanonical.Get();
+		FDynamicMesh3* Mesh = CurrentSelections[SelectionIndex].Target->UnwrapCanonical.Get();
+		const TArray<int32>& MovingVids = MovingVidsPerSelection[SelectionIndex];
 
-		for (int32 Eid : LivePreviewEids)
-		{
-			FVector3d Vert1, Vert2;
-			LivePreviewMesh->GetEdgeV(Eid, Vert1, Vert2);
-
-			LivePreviewLineSet->AddLine(
-				MeshTransform.TransformPosition(Vert1), 
-				MeshTransform.TransformPosition(Vert2), 
-				FUVEditorUXSettings::SelectionTriangleWireframeColor,
-				FUVEditorUXSettings::LivePreviewHighlightThickness,
-				FUVEditorUXSettings::LivePreviewHighlightDepthOffset);
-		}
+		MovingVertOriginalPositionsPerSelection.Emplace();
+		TArray<FVector3d>& MovingVertOriginalPositions = MovingVertOriginalPositionsPerSelection.Last();
+		MovingVertOriginalPositions.SetNum(MovingVids.Num());
 
-		for (int32 Vid : LivePreviewVids)
+		// Note: Our meshes currently don't have a transform. Otherwise we'd need to convert vid location to world
+		// space first, then to the frame.
+		for (int32 i = 0; i < MovingVids.Num(); ++i)
 		{
-			FVector3d Position = LivePreviewMesh->GetVertex(Vid);
-
-			LivePreviewPointSet->AddPoint(Position, 
-				FUVEditorUXSettings::SelectionTriangleWireframeColor,
-				FUVEditorUXSettings::LivePreviewHighlightPointSize,
-				FUVEditorUXSettings::LivePreviewHighlightDepthOffset);
+			MovingVertOriginalPositions[i] = InitialGizmoFrame.ToFramePoint(Mesh->GetVertex(MovingVids[i]));
 		}
 	}
 }
 
-void UUVSelectTool::GizmoTransformStarted(UTransformProxy* Proxy)
-{
-	bInDrag = true;
-
-	InitialGizmoFrame = FFrame3d(TransformGizmo->ActiveTarget->GetTransform());
-	MovingVertOriginalPositions.SetNum(MovingVids.Num());
-	const FDynamicMesh3* Mesh = Targets[SelectionTargetIndex]->UnwrapCanonical.Get();
-	// Note: Our meshes currently don't have a transform. Otherwise we'd need to convert vid location to world
-	// space first, then to the frame.
-	for (int32 i = 0; i < MovingVids.Num(); ++i)
-	{
-		MovingVertOriginalPositions[i] = InitialGizmoFrame.ToFramePoint(Mesh->GetVertex(MovingVids[i]));
-	}
-}
-
 void UUVSelectTool::GizmoTransformChanged(UTransformProxy* Proxy, FTransform Transform)
 {
 	// This function gets called both during drag and on undo/redo. This might have been ok if
@@ -851,385 +490,103 @@ void UUVSelectTool::GizmoTransformEnded(UTransformProxy* Proxy)
 {
 	bInDrag = false;
 
-	// Set things up for undo.
-	// TODO: We should really use FMeshVertexChange instead of FDynamicMeshChange because we don't
-	// need to alter the mesh topology. However we currently don't have a way to apply a FMeshVertexChange
-	// directly to a dynamic mesh pointer, only via UDynamicMesh. We should change things here once
-	// that ability exists.
-	FDynamicMeshChangeTracker ChangeTracker(Targets[SelectionTargetIndex]->UnwrapCanonical.Get());
-	ChangeTracker.BeginChange();
-	ChangeTracker.SaveTriangles(SelectedTids, true);
-
-	// One final attempt to apply transforms if OnTick hasn't happened yet
-	ApplyGizmoTransform();
-
-	// Both previews must already be updated, so only need to update canonical. 
-	{
-		// We don't want to react to the ensuing broadcast so that we don't lose the gizmo rotation. We could just 
-		// not broadcast (and update related structures, i.e. trees, ourselves), but conceptually it's better to 
-		// broadcast the change since we did change the canonicals.
-		TGuardValue<bool> IgnoreOnCanonicalBroadcast(bIgnoreOnCanonicalChange, true); // sets to true, restores at end
-
-		Targets[SelectionTargetIndex]->UpdateCanonicalFromPreviews(&MovingVids,
-			UUVEditorToolMeshInput::NONE_CHANGED_ARG);
-	}
-
 	const FText TransactionName(LOCTEXT("DragCompleteTransactionName", "Move Items"));
-	EmitChangeAPI->EmitToolIndependentChange(ChangeRouter, MakeUnique<UVSelectToolLocals::FGizmoMeshChange>(
-		Targets[SelectionTargetIndex], ChangeTracker.EndChange(), 
-		InitialGizmoFrame.ToFTransform(), TransformGizmo->GetGizmoTransform()),
-		TransactionName);
-
-	TransformGizmo->SetNewChildScale(FVector::One());
-	SelectionMechanic->RebuildDrawnElements(TransformGizmo->GetGizmoTransform());
-}
-
-void UUVSelectTool::ApplyGizmoTransform()
-{
-	if (bGizmoTransformNeedsApplication)
-	{
-		FTransformSRT3d TransformToApply(UnappliedGizmoTransform);
-
-		// TODO: The division here is a bit of a hack. Properly-speaking, the scaling handles should act relative to
-		// gizmo size, not the visible space across which we drag, otherwise it becomes dependent on the units we
-		// use and our absolute distance from the object. Since our UV unwrap is scaled by 1000 to make it
-		// easier to zoom in and out without running into issues, the measure of the distance across which we typically
-		// drag the handles is too high to be convenient. Until we make the scaling invariant to units/distance from
-		// target, we use this hack.
-		TransformToApply.SetScale(FVector::One() + (UnappliedGizmoTransform.GetScale3D() - FVector::One()) / 10);
-
-		Targets[SelectionTargetIndex]->UnwrapPreview->PreviewMesh->DeferredEditMesh([&TransformToApply, this](FDynamicMesh3& MeshIn)
-			{
-				for (int32 i = 0; i < MovingVids.Num(); ++i)
-				{
-					MeshIn.SetVertex(MovingVids[i], TransformToApply.TransformPosition(MovingVertOriginalPositions[i]));
-				}
-			}, false);
-		Targets[SelectionTargetIndex]->UpdateUnwrapPreviewOverlayFromPositions(&MovingVids, 
-			UUVEditorToolMeshInput::NONE_CHANGED_ARG, &SelectedTids);
-
-		SelectionMechanic->SetDrawnElementsTransform((FTransform)TransformToApply);
-
-		Targets[SelectionTargetIndex]->UpdateAppliedPreviewFromUnwrapPreview(&MovingVids, 
-			UUVEditorToolMeshInput::NONE_CHANGED_ARG, &SelectedTids);
-
-		bGizmoTransformNeedsApplication = false;
-		SewAction->UpdateVisualizations();
-		IslandConformalUnwrapAction->UpdateVisualizations();
-	}
-}
-
-void UUVSelectTool::Render(IToolsContextRenderAPI* RenderAPI)
-{
-	SelectionMechanic->Render(RenderAPI);
-}
-
-void UUVSelectTool::DrawHUD(FCanvas* Canvas, IToolsContextRenderAPI* RenderAPI)
-{
-	SelectionMechanic->DrawHUD(Canvas, RenderAPI);
-}
-
-void UUVSelectTool::OnTick(float DeltaTime)
-{
-	ApplyGizmoTransform();
+	EmitChangeAPI->BeginUndoTransaction(TransactionName);
 
-	// Deal with any buttons that may have been clicked
-	if (PendingAction != ESelectToolAction::NoAction)
+	for (int32 SelectionIndex = 0; SelectionIndex < CurrentSelections.Num(); ++SelectionIndex)
 	{
-		ApplyAction(PendingAction);
-		PendingAction = ESelectToolAction::NoAction;
-	}
-}
+		UUVEditorToolMeshInput* Target = CurrentSelections[SelectionIndex].Target.Get();
+		const TArray<int32>& MovingVids = MovingVidsPerSelection[SelectionIndex];
 
-void UUVSelectTool::RequestAction(ESelectToolAction ActionType)
-{
-	ClearWarning();
-	if (PendingAction == ESelectToolAction::NoAction)
-	{
-		PendingAction = ActionType;
-	}
-}
+		// Set things up for undo.
+		// TODO: We should really use FMeshVertexChange instead of FDynamicMeshChange because we don't
+		// need to alter the mesh topology. However we currently don't have a way to apply a FMeshVertexChange
+		// directly to a dynamic mesh pointer, only via UDynamicMesh. We should change things here once
+		// that ability exists.
+		FDynamicMeshChangeTracker ChangeTracker(Target->UnwrapCanonical.Get());
+		ChangeTracker.BeginChange();
+		ChangeTracker.SaveTriangles(RenderUpdateTidsPerSelection[SelectionIndex], true);
 
-void UUVSelectTool::ApplyAction(ESelectToolAction ActionType)
-{
-	auto MaybeAddAnalyticsActionHistoryItem = [this, &ActionType](FUVEditorDynamicMeshSelection::EType ExpectedType)
-	{
-		if (!SelectionMechanic->GetCurrentSelection().IsEmpty() && (SelectionMechanic->GetCurrentSelection().Type == ExpectedType))
+		// One final attempt to apply transforms if OnTick hasn't happened yet
+		if (bGizmoTransformNeedsApplication)
 		{
-			FActionHistoryItem ActionHistoryItem;
-			ActionHistoryItem.Timestamp = FDateTime::UtcNow();
-			ActionHistoryItem.NumOperands = SelectionMechanic->GetCurrentSelection().SelectedIDs.Num();
-			ActionHistoryItem.ActionType = ActionType;
-			AnalyticsActionHistory.Add(ActionHistoryItem);
+			ApplyGizmoTransform();
 		}
-	};
-	
-	switch (ActionType)
-	{
-	case ESelectToolAction::Sew:
-		if (SewAction)
-		{
-			TRACE_CPUPROFILER_EVENT_SCOPE(ApplyAction_Sew);
-	
-			MaybeAddAnalyticsActionHistoryItem(FUVEditorDynamicMeshSelection::EType::Edge);
-
-			const FText TransactionName(LOCTEXT("SewCompleteTransactionName", "Sew Edges"));
-			EmitChangeAPI->BeginUndoTransaction(TransactionName);
 
-			SelectionMechanic->SetSelection(FUVEditorDynamicMeshSelection(), false, true);
-			SewAction->ExecuteAction(*EmitChangeAPI);
-
-			EmitChangeAPI->EndUndoTransaction();
-		}
-		break;
-	case ESelectToolAction::IslandConformalUnwrap:
-		if (IslandConformalUnwrapAction)
+		// Both previews must already be updated, so only need to update canonical. 
 		{
-			TRACE_CPUPROFILER_EVENT_SCOPE(ApplyAction_IslandConformalUnwrap);
-			
-			MaybeAddAnalyticsActionHistoryItem(FUVEditorDynamicMeshSelection::EType::Triangle);
-
-			const FText TransactionName(LOCTEXT("ConformalUnwrapCompleteTransactionName", "Conformal Unwrap Islands"));
-			EmitChangeAPI->BeginUndoTransaction(TransactionName);
-
-			SelectionMechanic->SetSelection(FUVEditorDynamicMeshSelection(), false, true);
-			IslandConformalUnwrapAction->ExecuteAction(*EmitChangeAPI);
+			// We don't want to fully react to the ensuing broadcast so that we don't lose the gizmo rotation,
+			// but do want to broadcast so that other things like the trees update themselves.
+			TGuardValue<bool> GizmoUpdateGuard(bUpdateGizmoOnCanonicalChange, false); // sets to false, restores at end
 
-			EmitChangeAPI->EndUndoTransaction();
+			Target->UpdateCanonicalFromPreviews(&MovingVids, UUVEditorToolMeshInput::NONE_CHANGED_ARG);
 		}
-		break;
-	case ESelectToolAction::Split:
-		
-		MaybeAddAnalyticsActionHistoryItem(FUVEditorDynamicMeshSelection::EType::Edge);
-		
-		ApplySplit();
-	default:
-		break;
-	}
-}
 
-void UUVSelectTool::ApplySplit()
-{
-	using namespace UVSelectToolLocals;
+		EmitChangeAPI->EmitToolIndependentChange(GetToolManager(), MakeUnique<UVSelectToolLocals::FGizmoMeshChange>(
+			Target, ChangeTracker.EndChange(),
+			InitialGizmoFrame.ToFTransform(), TransformGizmo->GetGizmoTransform()),
+			TransactionName);
 
-	const FUVEditorDynamicMeshSelection& Selection = SelectionMechanic->GetCurrentSelection();
+		TransformGizmo->SetNewChildScale(FVector::One());
 
-	if (Selection.IsEmpty())
-	{
-		GetToolManager()->DisplayMessage(
-			LOCTEXT("SplitErrorSelectionEmpty", "Cannot split UVs. Selection was empty."),
-			EToolMessageLevel::UserWarning);
-	}
-	else if (Selection.Type == FUVEditorDynamicMeshSelection::EType::Edge)
-	{
-		ApplySplitEdges();
-	}
-	else if (Selection.Type == FUVEditorDynamicMeshSelection::EType::Vertex)
-	{
-		ApplySplitBowtieVertices();
-	}
-	else
-	{
-		GetToolManager()->DisplayMessage(
-			LOCTEXT("SplitErrorNotEdgeOrVert", "Cannot split UVs. Selection must be edges or vertices."),
-			EToolMessageLevel::UserWarning);
+		SelectionAPI->RebuildUnwrapHighlight(TransformGizmo->GetGizmoTransform());
 	}
+	EmitChangeAPI->EndUndoTransaction();
 }
 
-void UUVSelectTool::ApplySplitEdges()
+/**
+ * Updates previews based off of the current gizmo transform, CurrentSelections, MovingVidsPerSelection, 
+ * and MovingVertOriginalPositions. Gets emitted while dragging the gizmo on tick.
+ */
+void UUVSelectTool::ApplyGizmoTransform()
 {
-	const FUVEditorDynamicMeshSelection& Selection = SelectionMechanic->GetCurrentSelection();
-	if (!ensure(SelectionTargetIndex >= 0 && !Selection.IsEmpty() && Selection.Type == FUVEditorDynamicMeshSelection::EType::Edge))
-	{
-		return;
-	}
-	UUVEditorToolMeshInput* Target = Targets[SelectionTargetIndex];
+	FTransformSRT3d TransformToApply(UnappliedGizmoTransform);
 
-	// Gather up the corresponding edge IDs in the applied (3d) mesh.
-	TSet<int32> AppliedEidSet;
-	for (int32 Eid : Selection.SelectedIDs)
-	{
-		// Note that we don't check whether edges are already boundary edges because we allow such edges
-		// to be selected for splitting of any attached bowties.
-
-		FIndex2i EdgeUnwrapVids = Selection.Mesh->GetEdgeV(Eid);
+	// TODO: The division here is a bit of a hack. Properly-speaking, the scaling handles should act relative to
+	// gizmo size, not the visible space across which we drag, otherwise it becomes dependent on the units we
+	// use and our absolute distance from the object. Since our UV unwrap is scaled by 1000 to make it
+	// easier to zoom in and out without running into issues, the measure of the distance across which we typically
+	// drag the handles is too high to be convenient. Until we make the scaling invariant to units/distance from
+	// target, we use this hack.
+	TransformToApply.SetScale(FVector::One() + (UnappliedGizmoTransform.GetScale3D() - FVector::One()) / 10);
 
-		int32 AppliedEid = Target->AppliedCanonical->FindEdge(
-			Target->UnwrapVidToAppliedVid(EdgeUnwrapVids.A),
-			Target->UnwrapVidToAppliedVid(EdgeUnwrapVids.B));
+	for (int32 SelectionIndex = 0; SelectionIndex < CurrentSelections.Num(); ++SelectionIndex)
+	{
+		UUVEditorToolMeshInput* Target = CurrentSelections[SelectionIndex].Target.Get();
+		const TArray<int32>& MovingVids = MovingVidsPerSelection[SelectionIndex];
+		const TArray<FVector3d>& MovingVertOriginalPositions = MovingVertOriginalPositionsPerSelection[SelectionIndex];
 
-		if (ensure(AppliedEid != IndexConstants::InvalidID))
+		Target->UnwrapPreview->PreviewMesh->DeferredEditMesh([&TransformToApply, &MovingVids, &MovingVertOriginalPositions](FDynamicMesh3& MeshIn)
 		{
-			AppliedEidSet.Add(AppliedEid);
-		}
-	}
-
-	// Perform the cut in the overlay, but don't propagate to unwrap yet
-	FUVEditResult UVEditResult;
-	FDynamicMeshUVEditor UVEditor(Target->AppliedCanonical.Get(),
-		Target->UVLayerIndex, false);
-	UVEditor.CreateSeamsAtEdges(AppliedEidSet, &UVEditResult);
+			for (int32 i = 0; i < MovingVids.Num(); ++i)
+			{
+				MeshIn.SetVertex(MovingVids[i], TransformToApply.TransformPosition(MovingVertOriginalPositions[i]));
+			}
+		}, false);
 
-	// Figure out the triangles that need to be saved in the unwrap for undo
-	TSet<int32> TidSet;
-	for (int32 UnwrapVid : UVEditResult.NewUVElements)
-	{
-		TArray<int32> VertTids;
-		Target->AppliedCanonical->GetVtxTriangles(Target->UnwrapVidToAppliedVid(UnwrapVid), VertTids);
-		TidSet.Append(VertTids);
+		Target->UpdateUnwrapPreviewOverlayFromPositions(&MovingVids, UUVEditorToolMeshInput::NONE_CHANGED_ARG, 
+			&RenderUpdateTidsPerSelection[SelectionIndex]);
+		Target->UpdateAppliedPreviewFromUnwrapPreview(&MovingVids, UUVEditorToolMeshInput::NONE_CHANGED_ARG, 
+			&RenderUpdateTidsPerSelection[SelectionIndex]);
 	}
 
-	FDynamicMeshChangeTracker ChangeTracker(Target->UnwrapCanonical.Get());
-	ChangeTracker.BeginChange();
-	ChangeTracker.SaveTriangles(TidSet, true);
-
-	const FText TransactionName(LOCTEXT("ApplySplitEdgesTransactionName", "Split Edges"));
-	EmitChangeAPI->BeginUndoTransaction(TransactionName);
-
-	// Clear selection here so that it is restored on undo
-	// TODO: We could apply a new selection after the change, but this should only happen
-	// once we apply a new selection on split, which we currently don't do.
-	SelectionMechanic->SetSelection(FUVEditorDynamicMeshSelection(), true, true);
-
-	// Perform the update
-	TArray<int32> AppliedTids = TidSet.Array();
-	Target->UpdateAllFromAppliedCanonical(&UVEditResult.NewUVElements, &AppliedTids, &AppliedTids);
-	
-	// Not needed because it should happen automatically via broadcast of target canonical mesh change
-	// AABBTrees[SelectionTargetIndex]->Build();
-
-	// Emit update transaction
-	EmitChangeAPI->EmitToolIndependentUnwrapCanonicalChange(
-		Target, ChangeTracker.EndChange(), TransactionName);
-
-	EmitChangeAPI->EndUndoTransaction();
+	SelectionAPI->SetUnwrapHighlightTransform((FTransform)TransformToApply);
+	bGizmoTransformNeedsApplication = false;
 }
 
-void UUVSelectTool::ApplySplitBowtieVertices()
+void UUVSelectTool::Render(IToolsContextRenderAPI* RenderAPI)
 {
-	const FUVEditorDynamicMeshSelection& Selection = SelectionMechanic->GetCurrentSelection();
-	if (!ensure(SelectionTargetIndex >= 0 && !Selection.IsEmpty() && Selection.Type == FUVEditorDynamicMeshSelection::EType::Vertex))
-	{
-		return;
-	}
-	UUVEditorToolMeshInput* Target = Targets[SelectionTargetIndex];
-
-	// Gather the corresponding vert ID's in the applied mesh
-	TSet<int32> AppliedVidSet;
-	for (int32 UnwrapVid : Selection.SelectedIDs)
-	{
-		int32 AppliedVid = Target->UnwrapVidToAppliedVid(UnwrapVid);
-		AppliedVidSet.Add(AppliedVid);
-	}
-
-	// Split any bowties in the applied mesh overlay
-	TArray<int32> NewUVElements;
-	FDynamicMeshUVOverlay* Overlay = Target->AppliedCanonical->Attributes()->GetUVLayer(Target->UVLayerIndex);
-	for (int32 Vid : AppliedVidSet)
-	{
-		Overlay->SplitBowtiesAtVertex(Vid, &NewUVElements);
-	}
-
-	// Prep for undo transaction
-	TSet<int32> TidSet;
-	for (int32 UnwrapVid : NewUVElements)
-	{
-		TArray<int32> VertTids;
-		Target->AppliedCanonical->GetVtxTriangles(Target->UnwrapVidToAppliedVid(UnwrapVid), VertTids);
-		TidSet.Append(VertTids);
-	}
-
-	FDynamicMeshChangeTracker ChangeTracker(Target->UnwrapCanonical.Get());
-	ChangeTracker.BeginChange();
-	ChangeTracker.SaveTriangles(TidSet, true);
-
-	const FText TransactionName(LOCTEXT("ApplySplitBowtieVerticesTransactionName", "Split Bowties"));
-	EmitChangeAPI->BeginUndoTransaction(TransactionName);
-
-	// Emit selection clear first so that we restore it on undo
-	FUVEditorDynamicMeshSelection NewSelection = SelectionMechanic->GetCurrentSelection(); // save type, etc
-	// TODO: This emitted transaction doesn't actually need to broadcast on redo, but we don't yet
-	// have support for that.
-	SelectionMechanic->SetSelection(FUVEditorDynamicMeshSelection(), false, true); // don't broadcast, do emit undo
-
-	// Perform the update
-	TArray<int32> AppliedTids = TidSet.Array();
-	Target->UpdateAllFromAppliedCanonical(&NewUVElements, &AppliedTids, &AppliedTids);
-
-	// Emit update transaction
-	EmitChangeAPI->EmitToolIndependentUnwrapCanonicalChange(
-		Target, ChangeTracker.EndChange(), TransactionName);
-
-	// Set up the new selection to include the new elements
-	NewSelection.SelectedIDs.Append(NewUVElements);
-	// TODO: This emitted transaction doesn't actually need to broadcast on undo, but we don't yet
-	// have support for that.
-	SelectionMechanic->SetSelection(NewSelection, true, true); // both broadcast and emit undo
-
-	EmitChangeAPI->EndUndoTransaction();
 }
 
-void UUVSelectTool::RecordAnalytics()
+void UUVSelectTool::DrawHUD(FCanvas* Canvas, IToolsContextRenderAPI* RenderAPI)
 {
-	using namespace UVEditorAnalytics;
-	
-	if (!FEngineAnalytics::IsAvailable())
-	{
-		return;
-	}
-	
-	TArray<FAnalyticsEventAttribute> Attributes;
-	Attributes.Add(FAnalyticsEventAttribute(TEXT("Timestamp"), FDateTime::UtcNow().ToString()));
-	
-	// Tool inputs
-	InputTargetAnalytics.AppendToAttributes(Attributes, "Input");
-
-	// Tool outputs
-	const FTargetAnalytics OutputTargetAnalytics = CollectTargetAnalytics(Targets);
-	OutputTargetAnalytics.AppendToAttributes(Attributes, "Output");
-
-	// Tool stats
-	auto MaybeAppendActionStatsToAttributes = [this, &Attributes](ESelectToolAction ActionType, const FString& OperandName)
-	{
-		int32 NumActions = 0;
-		float MeanNumOperands = 0;
-		int32 MinNumOperands = TNumericLimits<int32>::Max();
-		int32 MaxNumOperands = TNumericLimits<int32>::Min();
-			
-		for (const FActionHistoryItem& Item : AnalyticsActionHistory)
-		{
-			if (Item.ActionType == ActionType)
-			{
-				NumActions += 1;
-				MeanNumOperands += static_cast<float>(Item.NumOperands);
-				MinNumOperands = FGenericPlatformMath::Min(MinNumOperands, Item.NumOperands);
-				MaxNumOperands = FGenericPlatformMath::Max(MaxNumOperands, Item.NumOperands);
-			}
-		}
-		
-		if (NumActions > 0)
-		{
-			MeanNumOperands /= NumActions;
-			const FString ActionName = StaticEnum<ESelectToolAction>()->GetNameStringByIndex(static_cast<int>(ActionType));
-			Attributes.Add(FAnalyticsEventAttribute(FString::Printf(TEXT("Stats.%sAction.NumActions"), *ActionName), NumActions));
-			Attributes.Add(FAnalyticsEventAttribute(FString::Printf(TEXT("Stats.%sAction.MinNum%s"), *ActionName, *OperandName), MinNumOperands));
-			Attributes.Add(FAnalyticsEventAttribute(FString::Printf(TEXT("Stats.%sAction.MaxNum%s"), *ActionName, *OperandName), MaxNumOperands));
-			Attributes.Add(FAnalyticsEventAttribute(FString::Printf(TEXT("Stats.%sAction.MeanNum%s"), *ActionName, *OperandName), MeanNumOperands));
-		}
-	};
-	MaybeAppendActionStatsToAttributes(ESelectToolAction::IslandConformalUnwrap, TEXT("Triangles"));
-	MaybeAppendActionStatsToAttributes(ESelectToolAction::Split, TEXT("Edges"));
-	MaybeAppendActionStatsToAttributes(ESelectToolAction::Sew, TEXT("Edges"));
-	Attributes.Add(FAnalyticsEventAttribute(TEXT("Stats.ToolActiveDuration"), (FDateTime::UtcNow() - ToolStartTimeAnalytics).ToString()));
-	
-	FEngineAnalytics::GetProvider().RecordEvent(UVEditorAnalyticsEventName(TEXT("EditTool")), Attributes);
+}
 
-	if constexpr (false)
+void UUVSelectTool::OnTick(float DeltaTime)
+{
+	if (bGizmoTransformNeedsApplication)
 	{
-		for (const FAnalyticsEventAttribute& Attr : Attributes)
-		{
-			UE_LOG(LogGeometry, Log, TEXT("Debug %s.EditTool.%s = %s"), *UVEditorAnalyticsPrefix, *Attr.GetName(), *Attr.GetValue());
-		}
+		ApplyGizmoTransform();
 	}
 }
 
diff --git a/Engine/Plugins/Experimental/UVEditor/Source/UVEditorTools/Private/UVToolContextObjects.cpp b/Engine/Plugins/Experimental/UVEditor/Source/UVEditorTools/Private/UVToolContextObjects.cpp
deleted file mode 100644
diff --git a/Engine/Plugins/Experimental/UVEditor/Source/UVEditorTools/Public/Actions/UVIslandConformalUnwrapAction.h b/Engine/Plugins/Experimental/UVEditor/Source/UVEditorTools/Public/Actions/UVIslandConformalUnwrapAction.h
new file mode 100644
--- /dev/null
+++ b/Engine/Plugins/Experimental/UVEditor/Source/UVEditorTools/Public/Actions/UVIslandConformalUnwrapAction.h
@@ -0,0 +1,23 @@
+// Copyright Epic Games, Inc. All Rights Reserved.
+
+#pragma once
+
+#include "CoreMinimal.h"
+#include "GeometryBase.h"
+#include "IndexTypes.h"
+
+#include "Actions/UVToolAction.h"
+
+#include "UVIslandConformalUnwrapAction.generated.h"
+
+PREDECLARE_GEOMETRY(class FUVEditorDynamicMeshSelection);
+
+UCLASS()
+class UVEDITORTOOLS_API UUVIslandConformalUnwrapAction : public UUVToolAction
+{
+	GENERATED_BODY()
+
+public:
+	virtual bool CanExecuteAction() const override;
+	virtual bool ExecuteAction() override;
+};
diff --git a/Engine/Plugins/Experimental/UVEditor/Source/UVEditorTools/Public/Actions/UVSeamSewAction.h b/Engine/Plugins/Experimental/UVEditor/Source/UVEditorTools/Public/Actions/UVSeamSewAction.h
new file mode 100644
--- /dev/null
+++ b/Engine/Plugins/Experimental/UVEditor/Source/UVEditorTools/Public/Actions/UVSeamSewAction.h
@@ -0,0 +1,34 @@
+// Copyright Epic Games, Inc. All Rights Reserved.
+
+#pragma once
+
+#include "CoreMinimal.h"
+
+#include "Actions/UVToolAction.h"
+#include "GeometryBase.h"
+#include "IndexTypes.h"
+
+#include "Actions/UVToolAction.h"
+
+#include "UVSeamSewAction.generated.h"
+
+PREDECLARE_GEOMETRY(class FUVEditorDynamicMeshSelection);
+PREDECLARE_GEOMETRY(class FDynamicMesh3);
+class APreviewGeometryActor;
+class ULineSetComponent;
+
+UCLASS()
+class UVEDITORTOOLS_API UUVSeamSewAction : public UUVToolAction
+{	
+	GENERATED_BODY()
+
+	using FDynamicMesh3 = UE::Geometry::FDynamicMesh3;
+
+public:
+	virtual bool CanExecuteAction() const override;
+	virtual bool ExecuteAction() override;
+
+	static int32 FindSewEdgeOppositePairing(const FDynamicMesh3& UnwrapMesh, 
+		const FDynamicMesh3& AppliedMesh, int32 UVLayerIndex, int32 UnwrapEid, 
+		bool& bWouldPreferOppositeOrderOut);
+};
\ No newline at end of file
diff --git a/Engine/Plugins/Experimental/UVEditor/Source/UVEditorTools/Public/Actions/UVSplitAction.h b/Engine/Plugins/Experimental/UVEditor/Source/UVEditorTools/Public/Actions/UVSplitAction.h
new file mode 100644
--- /dev/null
+++ b/Engine/Plugins/Experimental/UVEditor/Source/UVEditorTools/Public/Actions/UVSplitAction.h
@@ -0,0 +1,18 @@
+// Copyright Epic Games, Inc. All Rights Reserved.
+
+#pragma once
+
+#include "Actions/UVToolAction.h"
+#include "CoreMinimal.h"
+
+#include "UVSplitAction.generated.h"
+
+UCLASS()
+class UVEDITORTOOLS_API UUVSplitAction : public UUVToolAction
+{	
+	GENERATED_BODY()
+
+public:
+	virtual bool CanExecuteAction() const override;
+	virtual bool ExecuteAction() override;
+};
\ No newline at end of file
diff --git a/Engine/Plugins/Experimental/UVEditor/Source/UVEditorTools/Public/Actions/UVToolAction.h b/Engine/Plugins/Experimental/UVEditor/Source/UVEditorTools/Public/Actions/UVToolAction.h
new file mode 100644
--- /dev/null
+++ b/Engine/Plugins/Experimental/UVEditor/Source/UVEditorTools/Public/Actions/UVToolAction.h
@@ -0,0 +1,34 @@
+// Copyright Epic Games, Inc. All Rights Reserved.
+
+#pragma once
+
+#include "CoreMinimal.h"
+
+#include "UVToolAction.generated.h"
+
+class UUVToolEmitChangeAPI;
+class UUVToolSelectionAPI;
+class UInteractiveToolManager;
+
+UCLASS()
+class UVEDITORTOOLS_API UUVToolAction : public UObject
+{
+	GENERATED_BODY()
+public:
+
+	virtual void Setup(UInteractiveToolManager* ToolManager);
+	virtual void Shutdown();
+	virtual bool CanExecuteAction() const { return false;}
+	virtual bool ExecuteAction() { return false; }
+
+protected:
+	UPROPERTY()
+	TObjectPtr<UUVToolSelectionAPI> SelectionAPI = nullptr;
+
+	UPROPERTY()
+	TObjectPtr<UUVToolEmitChangeAPI> EmitChangeAPI = nullptr;
+
+	TWeakObjectPtr<UInteractiveToolManager> ToolManager;
+
+	virtual UInteractiveToolManager* GetToolManager() const { return ToolManager.Get(); }
+};
diff --git a/Engine/Plugins/Experimental/UVEditor/Source/UVEditorTools/Public/ContextObjects/UVToolAssetInputsContext.h b/Engine/Plugins/Experimental/UVEditor/Source/UVEditorTools/Public/ContextObjects/UVToolAssetInputsContext.h
new file mode 100644
--- /dev/null
+++ b/Engine/Plugins/Experimental/UVEditor/Source/UVEditorTools/Public/ContextObjects/UVToolAssetInputsContext.h
@@ -0,0 +1,34 @@
+// Copyright Epic Games, Inc. All Rights Reserved.
+
+#pragma once
+
+#include "CoreMinimal.h"
+
+#include "UVToolAssetInputsContext.generated.h"
+
+/**
+ * Simple context object used by UV asset editor to pass inputs down to the UV editor mode.
+ */
+UCLASS()
+class UVEDITORTOOLS_API UUVToolAssetInputsContext : public UUVToolContextObject
+{
+	GENERATED_BODY()
+public:
+
+	void Initialize(TArray<TObjectPtr<UObject>> AssetsIn, TArray<FTransform> TransformsIn)
+	{
+		Assets = MoveTemp(AssetsIn);
+		Transforms = MoveTemp(TransformsIn);
+		ensure(Transforms.Num() == Assets.Num());
+	}
+
+	void GetAssetInputs(TArray<TObjectPtr<UObject>>& AssetsOut, TArray<FTransform>& TransformsOut) const
+	{
+		AssetsOut = Assets;
+		TransformsOut = Transforms;
+	}
+
+protected:
+	TArray<TObjectPtr<UObject>> Assets;
+	TArray<FTransform> Transforms;
+};
\ No newline at end of file
diff --git a/Engine/Plugins/Experimental/UVEditor/Source/UVEditorTools/Public/ContextObjects/UVToolContextObjects.h b/Engine/Plugins/Experimental/UVEditor/Source/UVEditorTools/Public/ContextObjects/UVToolContextObjects.h
new file mode 100644
--- /dev/null
+++ b/Engine/Plugins/Experimental/UVEditor/Source/UVEditorTools/Public/ContextObjects/UVToolContextObjects.h
@@ -0,0 +1,294 @@
+// Copyright Epic Games, Inc. All Rights Reserved.
+
+#pragma once
+
+#include "CoreMinimal.h"
+
+#include "DynamicMesh/DynamicMeshAABBTree3.h"
+#include "DynamicMesh/DynamicMeshChangeTracker.h" // FDynamicMeshChange for TUniquePtr
+#include "InteractiveToolManager.h"
+#include "GeometryBase.h"
+
+#include "UVToolContextObjects.generated.h"
+
+// TODO: This should be spread out across multiple files
+
+PREDECLARE_GEOMETRY(class FDynamicMesh3);
+class FToolCommandChange;
+struct FViewCameraState;
+class UInputRouter;
+class UWorld;
+class UUVEditorToolMeshInput;
+
+/**
+ * Base class for context objects used in the UV editor.
+ */
+UCLASS()
+class UVEDITORTOOLS_API UUVToolContextObject : public UObject
+{
+	GENERATED_BODY()
+public:
+	/**
+	 * Called by the mode when shutting context objects down, allowing them to do any cleanup.
+	 * Initialization, on the other hand is usually done by some class-specific Initialize() method.
+	 */
+	virtual void Shutdown() {}
+
+	/**
+	 * Called whenever a tool is ended, for instance to let a context object remove listeners associated
+	 * with that tool (it shouldn't have to do so, but may choose to for robustness).
+	 */
+	virtual void OnToolEnded(UInteractiveTool* DeadTool) {}
+};
+
+/**
+ * An API object meant to be stored in a context object store that allows UV editor tools
+ * to emit appropriate undo/redo transactions.
+ */
+UCLASS()
+class UVEDITORTOOLS_API UUVToolEmitChangeAPI : public UUVToolContextObject
+{
+	GENERATED_BODY()
+public:
+
+	void Initialize(TObjectPtr<UInteractiveToolManager> ToolManagerIn)
+	{
+		ToolManager = ToolManagerIn;
+	}
+
+	virtual void BeginUndoTransaction(const FText& Description)
+	{
+		ToolManager->BeginUndoTransaction(Description);
+	}
+	virtual void EndUndoTransaction()
+	{
+		ToolManager->EndUndoTransaction();
+	}
+
+	/**
+	 * Emit a change that can be undone even if we leave the tool from which it is emitted (as
+	 * long as that UV editor instance is still open). 
+	 * Minor note: because we undo "out of" tools into a default tool and never out of a default tool,
+	 * in practice, tool-independent changes will only ever be applied/reverted in the same tool 
+	 * invocation that they were emitted or in the default tool, not in other arbitrary tools.
+	 * 
+	 * Since tool-independent changes usually operate on UV editor mesh input object, it is probably
+	 * preferable to use EmitToolIndependentUnwrapCanonicalChange, which will set up a proper transaction
+	 * for you.
+	 */
+	virtual void EmitToolIndependentChange(UObject* TargetObject, TUniquePtr<FToolCommandChange> Change, const FText& Description);
+
+	/**
+	 * A convenience function that is like EmitToolIndependentChange, but uses a FDynamicMeshChange
+	 * that operates on the UnwrapCanonical of an input to create a change object that updates the other
+	 * views and issues an OnUndoRedo broadcast on the input object.
+	 */
+	virtual void EmitToolIndependentUnwrapCanonicalChange(UUVEditorToolMeshInput* InputObject,
+		TUniquePtr<UE::Geometry::FDynamicMeshChange> UnwrapCanonicalMeshChange, const FText& Description);
+
+	/**
+	 * Emits a change that is considered expired when the active tool does not match the tool that was active
+	 * when it was emitted.
+	 */
+	virtual void EmitToolDependentChange(UObject* TargetObject, TUniquePtr<FToolCommandChange> Change, const FText& Description);
+
+protected:
+	TWeakObjectPtr<UInteractiveToolManager> ToolManager = nullptr;
+};
+
+/**
+ * Allows tools to interact with the 3d preview viewport, which has a separate
+ * world and input router.
+ */
+UCLASS()
+class UVEDITORTOOLS_API UUVToolLivePreviewAPI : public UUVToolContextObject
+{
+	GENERATED_BODY()
+public:
+
+	void Initialize(UWorld* WorldIn, UInputRouter* RouterIn,
+		TUniqueFunction<void(FViewCameraState& CameraStateOut)> GetLivePreviewCameraStateFuncIn);
+
+	UWorld* GetLivePreviewWorld() { return World.Get(); }
+	UInputRouter* GetLivePreviewInputRouter() { return InputRouter.Get(); }
+	void GetLivePreviewCameraState(FViewCameraState& CameraStateOut) 
+	{ 
+		if (GetLivePreviewCameraStateFunc)
+		{
+			GetLivePreviewCameraStateFunc(CameraStateOut);
+		}
+	}
+
+protected:
+	UPROPERTY()
+	TWeakObjectPtr<UWorld> World;
+
+	UPROPERTY()
+	TWeakObjectPtr<UInputRouter> InputRouter;
+
+	TUniqueFunction<void(FViewCameraState& CameraStateOut)> GetLivePreviewCameraStateFunc;
+};
+
+USTRUCT()
+struct UVEDITORTOOLS_API FUDIMBlock
+{
+	GENERATED_BODY();
+
+	UPROPERTY()
+	int32 UDIM = 1001;
+
+	int32 BlockU() const;
+	int32 BlockV() const;
+	void SetFromBlocks(int32 BlockU, int32 BlockV);
+};
+
+/**
+ * Allows tools to interact with the 2d preview viewport 
+ */
+UCLASS()
+class UVEDITORTOOLS_API UUVTool2DViewportAPI : public UUVToolContextObject
+{
+	GENERATED_BODY()
+public:
+
+	void SetUDIMBlocks(TArray<FUDIMBlock>& BlocksIn, bool bBroadcast = true)
+	{
+		UDIMBlocks = BlocksIn;
+		if (bBroadcast)
+		{
+			OnUDIMBlockChange.Broadcast(UDIMBlocks);
+		}
+	}
+
+	const TArray<FUDIMBlock>& GetUDIMBlocks() const
+	{
+		return UDIMBlocks;
+	}
+
+	void SetDrawGrid(bool bDrawGridIn, bool bBroadcast = true)
+	{
+		bDrawGrid = bDrawGridIn;
+		if (bBroadcast)
+		{
+			OnDrawGridChange.Broadcast(bDrawGrid);
+		}
+	}
+
+	const bool GetDrawGrid() const
+    {
+		return bDrawGrid;
+	}
+
+	void SetDrawRulers(bool bDrawRulersIn, bool bBroadcast = true)
+	{
+		bDrawRulers = bDrawRulersIn;
+		if (bBroadcast)
+		{
+			OnDrawRulersChange.Broadcast(bDrawRulers);
+		}
+	}
+
+	const bool GetDrawRulers() const
+    {
+		return bDrawRulers;
+	}
+
+
+	DECLARE_MULTICAST_DELEGATE_OneParam(FOnUDIMBlockChange, const TArray<FUDIMBlock>& UDIMBlocks);
+	FOnUDIMBlockChange OnUDIMBlockChange;
+
+	DECLARE_MULTICAST_DELEGATE_OneParam(FOnDrawGridChange, bool bDrawGrid);
+	FOnDrawGridChange OnDrawGridChange;
+
+	DECLARE_MULTICAST_DELEGATE_OneParam(FOnDrawRulersChange, bool bDrawRulers);
+	FOnDrawRulersChange OnDrawRulersChange;
+
+protected:
+
+	UPROPERTY(Transient)
+	TArray<FUDIMBlock> UDIMBlocks;
+
+	UPROPERTY(Transient)
+	bool bDrawGrid;
+
+	UPROPERTY(Transient)
+	bool bDrawRulers;
+};
+
+
+/**
+ * Allows tools to interact with the assets and their UV layers
+*/
+UCLASS()
+class UVEDITORTOOLS_API UUVToolAssetAndChannelAPI : public UUVToolContextObject
+{
+	GENERATED_BODY()
+public:
+
+	TArray<int32> GetCurrentChannelVisibility()
+	{
+		if (GetCurrentChannelVisibilityFunc)
+		{
+			return GetCurrentChannelVisibilityFunc();
+		}
+		return TArray<int32>();
+	}
+
+	void RequestChannelVisibilityChange(const TArray<int32>& ChannelPerAsset, bool bEmitUndoTransaction=true)
+	{
+		if (RequestChannelVisibilityChangeFunc)
+		{
+			RequestChannelVisibilityChangeFunc(ChannelPerAsset, bEmitUndoTransaction);
+		}
+	}
+
+	void NotifyOfAssetChannelCountChange(int32 AssetID)
+	{
+		if (NotifyOfAssetChannelCountChangeFunc)
+		{
+			NotifyOfAssetChannelCountChangeFunc(AssetID);
+		}
+	}
+
+
+	TUniqueFunction<TArray<int32>()> GetCurrentChannelVisibilityFunc;
+	TUniqueFunction<void(const TArray<int32>&, bool)> RequestChannelVisibilityChangeFunc;
+	TUniqueFunction<void(int32 AssetID)> NotifyOfAssetChannelCountChangeFunc;
+
+};
+
+/** 
+ * Stores AABB trees for UV input object unwrap canonical or applied canonical meshes.
+ * Binds to the input objects' OnCanonicalModified delegate to automatically update 
+ * the tree when necessary.
+ */
+UCLASS()
+class UVEDITORTOOLS_API UUVToolAABBTreeStorage : public UUVToolContextObject
+{
+	GENERATED_BODY()
+public:
+	using FDynamicMesh3 = UE::Geometry::FDynamicMesh3;
+	using FDynamicMeshAABBTree3 = UE::Geometry::FDynamicMeshAABBTree3;
+
+
+	void Set(FDynamicMesh3* MeshKey, TSharedPtr<FDynamicMeshAABBTree3> Tree,
+		UUVEditorToolMeshInput* InputObject);
+
+	TSharedPtr<FDynamicMeshAABBTree3> Get(FDynamicMesh3* MeshKey);
+
+	void Remove(FDynamicMesh3* MeshKey);
+
+	void RemoveByPredicate(TUniqueFunction<
+		bool(const FDynamicMesh3* Mesh, TWeakObjectPtr<UUVEditorToolMeshInput> InputObject,
+			TSharedPtr<FDynamicMeshAABBTree3> Tree)> Predicate);
+
+	void Empty();
+
+	virtual void Shutdown() override;
+
+protected:
+
+	using TreeInputObjectPair = TPair<TSharedPtr<FDynamicMeshAABBTree3>, TWeakObjectPtr<UUVEditorToolMeshInput>>;
+
+	TMap<FDynamicMesh3*, TreeInputObjectPair> AABBTreeStorage;
+};
diff --git a/Engine/Plugins/Experimental/UVEditor/Source/UVEditorTools/Public/ContextObjects/UVToolViewportButtonsAPI.h b/Engine/Plugins/Experimental/UVEditor/Source/UVEditorTools/Public/ContextObjects/UVToolViewportButtonsAPI.h
new file mode 100644
--- /dev/null
+++ b/Engine/Plugins/Experimental/UVEditor/Source/UVEditorTools/Public/ContextObjects/UVToolViewportButtonsAPI.h
@@ -0,0 +1,95 @@
+// Copyright Epic Games, Inc. All Rights Reserved.
+
+#pragma once
+
+#include "CoreMinimal.h"
+
+#include "Selection/UVToolSelectionAPI.h"
+
+#include "UVToolViewportButtonsAPI.generated.h"
+
+/**
+ * Allows tools to interact with buttons in the viewport
+ */
+UCLASS()
+class UVEDITORTOOLS_API UUVToolViewportButtonsAPI : public UUVToolContextObject
+{
+	GENERATED_BODY()
+public:
+
+	enum class EGizmoMode
+	{
+		Select,
+		Transform
+	};
+
+	using ESelectionMode = UUVToolSelectionAPI::EUVEditorSelectionMode;
+
+	void SetGizmoButtonsEnabled(bool bOn)
+	{
+		bGizmoButtonsEnabled = bOn;
+	}
+
+	bool AreGizmoButtonsEnabled()
+	{
+		return bGizmoButtonsEnabled;
+	}
+
+	void SetGizmoMode(EGizmoMode ModeIn, bool bBroadcast = true)
+	{
+		GizmoMode = ModeIn;
+		if (bBroadcast)
+		{
+			OnGizmoModeChange.Broadcast(GizmoMode);
+		}
+	}
+
+	EGizmoMode GetGizmoMode()
+	{
+		return GizmoMode;
+	}
+
+
+	void SetSelectionButtonsEnabled(bool bOn)
+	{
+		bSelectionButtonsEnabled = bOn;
+	}
+
+	bool AreSelectionButtonsEnabled()
+	{
+		return bSelectionButtonsEnabled;
+	}
+
+	void SetSelectionMode(ESelectionMode ModeIn, bool bBroadcast = true)
+	{
+		SelectionMode = ModeIn;
+		if (bBroadcast)
+		{
+			OnSelectionModeChange.Broadcast(SelectionMode);
+		}
+	}
+
+	ESelectionMode GetSelectionMode()
+	{
+		return SelectionMode;
+	}
+
+	DECLARE_MULTICAST_DELEGATE_OneParam(FOnGizmoModeChange, EGizmoMode NewGizmoMode);
+	FOnGizmoModeChange OnGizmoModeChange;
+
+	DECLARE_MULTICAST_DELEGATE_OneParam(FOnSelectionModeChange, ESelectionMode NewSelectionMode);
+	FOnSelectionModeChange OnSelectionModeChange;
+
+	virtual void OnToolEnded(UInteractiveTool* DeadTool) override 
+	{
+		OnGizmoModeChange.RemoveAll(DeadTool);
+		OnSelectionModeChange.RemoveAll(DeadTool);
+	}
+
+protected:
+
+	bool bGizmoButtonsEnabled = false;
+	EGizmoMode GizmoMode = EGizmoMode::Select;
+	bool bSelectionButtonsEnabled = false;
+	ESelectionMode SelectionMode = ESelectionMode::Island;
+};
\ No newline at end of file
diff --git a/Engine/Plugins/Experimental/UVEditor/Source/UVEditorTools/Public/Selection/UVEditorDynamicMeshSelection.h b/Engine/Plugins/Experimental/UVEditor/Source/UVEditorTools/Public/Selection/UVEditorDynamicMeshSelection.h
deleted file mode 100644
diff --git a/Engine/Plugins/Experimental/UVEditor/Source/UVEditorTools/Public/Selection/UVEditorMeshSelectionMechanic.h b/Engine/Plugins/Experimental/UVEditor/Source/UVEditorTools/Public/Selection/UVEditorMeshSelectionMechanic.h
--- a/Engine/Plugins/Experimental/UVEditor/Source/UVEditorTools/Public/Selection/UVEditorMeshSelectionMechanic.h
+++ b/Engine/Plugins/Experimental/UVEditor/Source/UVEditorTools/Public/Selection/UVEditorMeshSelectionMechanic.h
@@ -9,156 +9,120 @@
 #include "DynamicMesh/DynamicMeshAABBTree3.h"
 #include "InteractionMechanic.h"
 #include "InteractiveTool.h"
-#include "Selection/UVEditorDynamicMeshSelection.h"
+#include "Selection/UVToolSelection.h"
+#include "Selection/UVToolSelectionAPI.h" // EUVEditorSelectionMode
 #include "Mechanics/RectangleMarqueeMechanic.h"
 #include "ToolContextInterfaces.h" //FViewCameraState
 
 #include "UVEditorMeshSelectionMechanic.generated.h"
 
-class UTriangleSetComponent;
-class ULineSetComponent;
-class UPointSetComponent;
 class APreviewGeometryActor;
-class UMaterialInstanceDynamic;
 struct FCameraRectangle;
+class ULineSetComponent;
+class UMaterialInstanceDynamic;
+class UPointSetComponent;
+class UTriangleSetComponent;
+class UUVToolViewportButtonsAPI;
 
-enum class EUVEditorMeshSelectionMode
-{
-	Component,
-	
-	// Not yet fully implemented for UV mesh purposes, since 
-	// we need to be able to select occluded edges
-	Edge,
-	Vertex,
-	Triangle,
-
-	// TODO: This might be good to rename later. And determine how it might interact with multi-mesh selection?
-	Mesh
-};
-
-UCLASS()
-class UVEDITORTOOLS_API UUVEditorMeshSelectionMechanicProperties : public UInteractiveToolPropertySet
-{
-	GENERATED_BODY()
-
-public:
-	/** If true, show the currently hovered selectable primitive */
-	UPROPERTY(EditAnywhere, Category = Options)
-	bool bShowHoveredElements = true;
-};
 
 /**
- * Mechanic for selecting elements of a dynamic mesh in the UV editor.
- * 
- * TODO: Currently only able to select unoccluded elements.
+ * Mechanic for selecting elements of a dynamic mesh in the UV editor. Interacts
+ * heavily with UUVToolSelectionAPI, which actually stores selections.
  */
 UCLASS()
-class UVEDITORTOOLS_API UUVEditorMeshSelectionMechanic : public UInteractionMechanic, public IClickBehaviorTarget
+class UVEDITORTOOLS_API UUVEditorMeshSelectionMechanic : public UInteractionMechanic, 
+	public IClickBehaviorTarget,
+	public IHoverBehaviorTarget
 {
 	GENERATED_BODY()
 
 public:
 	using FDynamicMeshAABBTree3 = UE::Geometry::FDynamicMeshAABBTree3;
-	using FUVEditorDynamicMeshSelection = UE::Geometry::FUVEditorDynamicMeshSelection;
+	using FUVToolSelection = UE::Geometry::FUVToolSelection;
 
 	virtual ~UUVEditorMeshSelectionMechanic() {}
 
 	virtual void Setup(UInteractiveTool* ParentTool) override;
 	virtual void Shutdown() override;
 
-	void SetWorld(UWorld* World);
+	// Initialization functions.
+	// The selection API is provided as a parameter rather than being grabbed out of the context
+	// store mainly because UVToolSelectionAPI itself sets up a selection mechanic, and is not 
+	// yet in the context store when it does this. 
+	void Initialize(UWorld* World, UUVToolSelectionAPI* SelectionAPI);
+	void SetTargets(const TArray<TObjectPtr<UUVEditorToolMeshInput>>& TargetsIn);
 
-	// Use this to initialize the meshes we want to hit test.
-	virtual void AddSpatial(TSharedPtr<FDynamicMeshAABBTree3> SpatialIn, const FTransform& TransformIn);
+	void SetIsEnabled(bool bIsEnabled);
+	bool IsEnabled() { return bIsEnabled; };
 
-	FVector3d GetCurrentSelectionCentroid();
+	void SetShowHoveredElements(bool bShow);
 
-	// Rebuilds the drawn selection highlights, and intializes them in such a way that their transform
-	// is equal to StartTransform (useful so that their transform can later be changed)
-	void RebuildDrawnElements(const FTransform& StartTransform);
-
-	// Changes the transform of the selection highlights. Useful for quickly updating the hightlight
-	// without rebuilding it, when the change is a transformation.
-	void SetDrawnElementsTransform(const FTransform& Transform);
-
-	virtual const FUVEditorDynamicMeshSelection& GetCurrentSelection() const;
-	void ChangeSelectionMode(const EUVEditorMeshSelectionMode& TargetMode);
-	virtual void SetSelection(const FUVEditorDynamicMeshSelection& Selection, bool bBroadcast = false, bool bEmitChange = false);
+	using ESelectionMode = UUVToolSelectionAPI::EUVEditorSelectionMode;
+	using FModeChangeOptions = UUVToolSelectionAPI::FSelectionMechanicModeChangeOptions;
+	/**
+	 * Sets selection mode for the mechanic.
+	 */
+	void SetSelectionMode(ESelectionMode TargetMode,
+		const FModeChangeOptions& Options = FModeChangeOptions());
 
 	virtual void Render(IToolsContextRenderAPI* RenderAPI) override;
 	virtual void DrawHUD(FCanvas* Canvas, IToolsContextRenderAPI* RenderAPI);
 
-	void OnDragRectangleStarted();
-	void OnDragRectangleChanged(const FCameraRectangle& CurrentRectangle);
-	void OnDragRectangleFinished(const FCameraRectangle& Rectangle, bool bCancelled);
-
 	// IClickBehaviorTarget implementation
 	virtual FInputRayHit IsHitByClick(const FInputDeviceRay& ClickPos) override;
 	virtual void OnClicked(const FInputDeviceRay& ClickPos) override;
 
 	// IModifierToggleBehaviorTarget implementation
 	virtual void OnUpdateModifierState(int ModifierID, bool bIsOn) override;
-	
-	TSet<int32> RayCast(const FInputDeviceRay& ClickPos, EUVEditorMeshSelectionMode Mode);
 
-	FSimpleMulticastDelegate OnSelectionChanged;
-
-	EUVEditorMeshSelectionMode SelectionMode;
+	// IHoverBehaviorTarget implementation
+	virtual FInputRayHit BeginHoverSequenceHitTest(const FInputDeviceRay& PressPos) override;
+	virtual void OnBeginHover(const FInputDeviceRay& DevicePos) override;
+	virtual bool OnUpdateHover(const FInputDeviceRay& DevicePos) override;
+	virtual void OnEndHover() override;
 
 	/**
-	 * Function to use for emitting selection change events. If not set, Setup() will attach a version that
-	 * uses EmitObjectChange() on the tool manager to emit a change that operates on this mechanic.
-	 * TODO: User should probably be able to specify whether the emitted transaction should broadcast
-	 * OnSelectionChanged on redo, undo, or both, to allow selection change events to be used as bookends
-	 * around topology changes.
+	 * Broadcasted whenever the marquee mechanic rectangle is changed, since these changes
+	 * don't trigger normal selection broadcasts.
 	 */ 
-	TUniqueFunction<void(const FUVEditorDynamicMeshSelection& OldSelection, const FUVEditorDynamicMeshSelection& NewSelection)> EmitSelectionChange;
+	FSimpleMulticastDelegate OnDragSelectionChanged;
 
 protected:
 
-	// All four combinations of shift/ctrl down are assigned a behaviour
-	bool ShouldAddToSelection() const { return !bCtrlToggle && bShiftToggle; }
-	bool ShouldRemoveFromSelection() const { return bCtrlToggle && !bShiftToggle; }
-	bool ShouldToggleFromSelection() const { return bCtrlToggle && bShiftToggle; }
-	bool ShouldRestartSelection() const { return !bCtrlToggle && !bShiftToggle; }
-
 	UPROPERTY()
-	TObjectPtr<UUVEditorMeshSelectionMechanicProperties> Settings = nullptr;
+	TObjectPtr<UUVToolSelectionAPI> SelectionAPI = nullptr;
 
 	UPROPERTY()
-	TObjectPtr<URectangleMarqueeMechanic> MarqueeMechanic;
+	TObjectPtr<UUVToolViewportButtonsAPI> ViewportButtonsAPI = nullptr;
 
 	UPROPERTY()
-	TObjectPtr<APreviewGeometryActor> PreviewGeometryActor = nullptr;
+	TObjectPtr<UUVToolEmitChangeAPI> EmitChangeAPI = nullptr;
 
-	/** The material being displayed for selected triangles */
-	UPROPERTY()
-	TObjectPtr<UMaterialInstanceDynamic> TriangleSetMaterial = nullptr;
-	UPROPERTY()
-	TObjectPtr<UTriangleSetComponent> TriangleSet = nullptr;
-	UPROPERTY()
-	TObjectPtr<ULineSetComponent> LineSet = nullptr;
 	UPROPERTY()
-	TObjectPtr<UPointSetComponent> PointSet = nullptr;
-	
-	UPROPERTY()
-	TObjectPtr<APreviewGeometryActor> HoverGeometryActor = nullptr;
-	
+	TObjectPtr<URectangleMarqueeMechanic> MarqueeMechanic = nullptr;
+
 	UPROPERTY()
 	TObjectPtr<UMaterialInstanceDynamic> HoverTriangleSetMaterial = nullptr;
+
 	UPROPERTY()
-	TObjectPtr<UTriangleSetComponent> HoverTriangleSet = nullptr;
-	UPROPERTY()
-	TObjectPtr<ULineSetComponent> HoverLineSet = nullptr;
-	UPROPERTY()
-	TObjectPtr<UPointSetComponent> HoverPointSet = nullptr;
-	
-
-	TArray<TSharedPtr<FDynamicMeshAABBTree3>> MeshSpatials;
-	TArray<FTransform> MeshTransforms;
-	FUVEditorDynamicMeshSelection CurrentSelection;
-	FUVEditorDynamicMeshSelection PreDragSelection;
-	int32 CurrentSelectionIndex = IndexConstants::InvalidID;
+	TObjectPtr<APreviewGeometryActor> HoverGeometryActor = nullptr;
+	// Weak pointers so that they go away when geometry actor is destroyed
+	TWeakObjectPtr<UTriangleSetComponent> HoverTriangleSet = nullptr;
+	TWeakObjectPtr<ULineSetComponent> HoverLineSet = nullptr;
+	TWeakObjectPtr<UPointSetComponent> HoverPointSet = nullptr;
+
+	// Should be the same as the mode-level targets array, indexed by AssetID
+	TArray<TObjectPtr<UUVEditorToolMeshInput>> Targets;
+	TArray<TSharedPtr<FDynamicMeshAABBTree3>> MeshSpatials; // 1:1 with Targets
+
+	ESelectionMode SelectionMode;
+	bool bIsEnabled = false;
+	bool bShowHoveredElements = true;
+
+	bool GetHitTid(const FInputDeviceRay& ClickPos, int32& TidOut,
+		int32& AssetIDOut, int32* ExistingSelectionObjectIndexOut = nullptr);
+	void ModifyExistingSelection(TSet<int32>& SelectionSetToModify, const TArray<int32>& SelectedIDs);
+
 	FViewCameraState CameraState;
 
 	bool bShiftToggle = false;
@@ -166,12 +130,19 @@ protected:
 	static const int32 ShiftModifierID = 1;
 	static const int32 CtrlModifierID = 2;
 
-	FVector3d CurrentSelectionCentroid;
-	int32 CentroidTimestamp = -1;
-	void UpdateCentroid();
+	// All four combinations of shift/ctrl down are assigned a behaviour
+	bool ShouldAddToSelection() const { return !bCtrlToggle && bShiftToggle; }
+	bool ShouldRemoveFromSelection() const { return bCtrlToggle && !bShiftToggle; }
+	bool ShouldToggleFromSelection() const { return bCtrlToggle && bShiftToggle; }
+	bool ShouldRestartSelection() const { return !bCtrlToggle && !bShiftToggle; }
 
-private:
-	void ClearCurrentSelection();
-	void UpdateCurrentSelection(const TSet<int32>& NewSelection, bool CalledFromOnRectangleChanged = false);
+	// For marquee mechanic
+	void OnDragRectangleStarted();
+	void OnDragRectangleChanged(const FCameraRectangle& CurrentRectangle);
+	void OnDragRectangleFinished(const FCameraRectangle& Rectangle, bool bCancelled);
+	TArray<FUVToolSelection> PreDragSelections;
+	// Maps asset id to a pre drag selection so that it is easy to tell which assets
+	// started with a selection. 1:1 with Targets.
+	TArray<const FUVToolSelection*> AssetIDToPreDragSelection;
 };
 
diff --git a/Engine/Plugins/Experimental/UVEditor/Source/UVEditorTools/Public/Selection/UVToolSelection.h b/Engine/Plugins/Experimental/UVEditor/Source/UVEditorTools/Public/Selection/UVToolSelection.h
new file mode 100644
--- /dev/null
+++ b/Engine/Plugins/Experimental/UVEditor/Source/UVEditorTools/Public/Selection/UVToolSelection.h
@@ -0,0 +1,74 @@
+// Copyright Epic Games, Inc. All Rights Reserved.
+
+#pragma once
+
+#include "CoreMinimal.h"
+
+#include "ToolTargets/UVEditorToolMeshInput.h"
+#include "Selection/StoredMeshSelectionUtil.h" //FEdgesFromTriangleSubIndices
+#include "DynamicMesh/DynamicMesh3.h"
+#include "Util/IndexUtil.h"
+
+namespace UE {
+namespace Geometry {
+
+/**
+ * Class that represents a selection in the canonical unwrap of a UV editor input object.
+ */
+class UVEDITORTOOLS_API FUVToolSelection
+{
+public:
+	enum class EType
+	{
+		Vertex,
+		Edge,
+		Triangle
+	};
+
+	TWeakObjectPtr<UUVEditorToolMeshInput> Target = nullptr;
+	EType Type = EType::Vertex;
+	TSet<int32> SelectedIDs;
+
+	void Clear()
+	{
+		Target = nullptr;
+		SelectedIDs.Reset();
+		StableEdgeIDs.Reset();
+	}
+
+	bool IsEmpty() const
+	{
+		return SelectedIDs.IsEmpty();
+	}
+
+	bool HasStableEdgeIdentifiers() const
+	{
+		return !StableEdgeIDs.IsEmpty();
+	}
+
+	bool operator==(const FUVToolSelection& Other) const
+	{
+		return Target == Other.Target
+			&& Type == Other.Type
+			&& SelectedIDs.Num() == Other.SelectedIDs.Num()
+			// Don't need to check the reverse because we checked Num above.
+			&& SelectedIDs.Includes(Other.SelectedIDs); 
+	}
+
+	bool operator!=(const FUVToolSelection& Other) const
+	{
+		return !(*this == Other);
+	}
+
+	void SaveStableEdgeIdentifiers(const FDynamicMesh3& Mesh);
+
+	void RestoreFromStableEdgeIdentifiers(const FDynamicMesh3& Mesh);
+
+	bool AreElementsPresentInMesh(FDynamicMesh3& Mesh) const;
+
+protected:
+	
+	FMeshEdgesFromTriangleSubIndices StableEdgeIDs;
+};
+
+}} //end namespace UE::Geometry
\ No newline at end of file
diff --git a/Engine/Plugins/Experimental/UVEditor/Source/UVEditorTools/Public/Selection/UVToolSelectionAPI.h b/Engine/Plugins/Experimental/UVEditor/Source/UVEditorTools/Public/Selection/UVToolSelectionAPI.h
new file mode 100644
--- /dev/null
+++ b/Engine/Plugins/Experimental/UVEditor/Source/UVEditorTools/Public/Selection/UVToolSelectionAPI.h
@@ -0,0 +1,329 @@
+// Copyright Epic Games, Inc. All Rights Reserved.
+
+#pragma once
+
+#include "CoreMinimal.h"
+
+#include "ContextObjects/UVToolContextObjects.h"
+#include "Selection/UVToolSelection.h"
+
+#include "UVToolSelectionAPI.generated.h"
+
+class UUVEditorMechanicAdapterTool;
+class UUVEditorMeshSelectionMechanic;
+class UUVToolSelectionHighlightMechanic;
+
+/**
+ * API for dealing with mode-level selection in the UV editor.
+ * 
+ * Selections are stored in a list of objects, one object per asset that contains a
+ * selection, all of the same type (vert/edge/tri), none empty. Selections are
+ * considered to be referring to the UnwrapCanonical mesh of the corresponding target.
+ * 
+ * There are also functions to enable automatic highlighting of the current selection,
+ * and to enable a selection mechanic in the viewport (to which tools can respond via
+ * OnSelectionChanged broadcasts).
+ */
+UCLASS()
+class UVEDITORTOOLS_API UUVToolSelectionAPI : public UUVToolContextObject
+{
+	GENERATED_BODY()
+public:
+	using FUVToolSelection = UE::Geometry::FUVToolSelection;
+
+	/**
+	 * Returns true when current selection is not empty.
+	 */
+	bool HaveSelections() const { return CurrentSelections.Num() > 0; }
+
+	/**
+	 * Returns the current selection. When there are multiple selection objects, it means that there
+	 * are selections in multiple assets, with a separate selection object per asset containing
+	 * a selection. In such a case, all objects will have the same type (vertex/edge/triangle).
+	 * 
+	 * Each selection object is considered to be referring to the CanonicalUnwrap mesh of the
+	 * associated target.
+	 */
+	const TArray<FUVToolSelection>& GetSelections() const { return CurrentSelections; }
+
+	/**
+	 * Returns type (vertex/edge/triangle) of current selection. Undefined if selection is empty.
+	 */
+	FUVToolSelection::EType GetSelectionsType()
+	{
+		return CurrentSelections.Num() > 0 ? CurrentSelections[0].Type : FUVToolSelection::EType::Triangle;
+	}
+
+	/**
+	 * Sets the current selection. Selections should be same type, with no more than one selection
+	 * object per asset, and no empty selection objects. Selections should be referring to the
+	 * UnwrapCanonical of the associated target.
+	 * 
+	 * @param bBroadcast If true, broadcast OnPreSelectionChange and OnSelectionChanged
+	 * @param bEmitChange If true, emit an undo/redo transaction.
+	 */
+	void SetSelections(const TArray<FUVToolSelection>& SelectionsIn, bool bBroadcast, bool bEmitChange);
+
+	/**
+	 * Equivalent to calling SetSelections with an empty array.
+	 */
+	void ClearSelections(bool bBroadcast, bool bEmitChange);
+
+	/**
+	 * Gets the centroid of the current selection. The centroid is cached and
+	 * invalidated on the next SetSelection call, but can be forced to be
+	 * recalculated.
+	 */
+	FVector3d GetUnwrapSelectionCentroid(bool bForceRecalculate = false);
+
+	// TODO: When we have support for selection in 3d live preview viewport, we may need these.
+	// These are function that would be used if a selection made in the applied mesh cannot be
+	// converted to an unwrap mesh selection due to unset UV elements (meaning that the selected
+	// portion does not exist in the unwrap mesh). Such selections would be stored separately 
+	// from all those that are able to be converted (which get stored with the usual unwrap selections).
+	// These "unset element" selections would be ignored by many tools/actions but operated on 
+	// by some others.
+	//bool HaveUnsetElementAppliedMeshSelections() const;
+	//const TArray<FUVToolSelection>& GetUnsetElementAppliedMeshSelections();
+	//void SetUnsetElementAppliedMeshSelections(const TArray<FUVToolSelection>& SelectionsIn, bool bBroadcast, bool bEmitChange);
+	//void ClearUnsetElementAppliedMeshSelections(bool bBroadcast, bool bEmitChange);
+
+	/**
+	 * Mode of operation for the selection mechanic.
+	 */
+	enum class EUVEditorSelectionMode
+	{
+		// When None, means that mechanic will not select anything
+		None,
+
+		Vertex,
+		Edge,
+		Triangle,
+		Island,
+		Mesh,
+	};
+
+	// Selection mechanic controls:
+
+	void SetSelectionMechanicEnabled(bool bEnabled);
+
+	struct FSelectionMechanicOptions
+	{
+		bool bShowHoveredElements = true;
+	};
+	void SetSelectionMechanicOptions(const FSelectionMechanicOptions& Options);
+
+	// Options for how a SetSelectionMechanicMode() call is performed
+	struct FSelectionMechanicModeChangeOptions
+	{
+		/**
+		 * If true, any existing selection will be converted to be compatible with the new mode.
+		 * For instance if an existing selection is vertices and we switch to Island selection
+		 * mode, the vertex selection will be converted to a triangle selection.
+		 */
+		bool bConvertExisting = true;
+
+		/**
+		 * If true and a conversion is performed (requires bConvertExisting to be true), the
+		 * OnPreSelectionChanged and OnSelectionChanged delegates will be broadcast.
+		 */
+		bool bBroadcastIfConverted = true;
+
+		/** If true, emit appropriate undo / redo transactions. */
+		bool bEmitChanges = true;
+
+		// Shouldn't really need this constructor but there seems to be a clang bug that requires it if we
+		// want to use a default-constructed instance as a default arg as we do below.
+		FSelectionMechanicModeChangeOptions() {}
+	};
+	/**
+	 * Sets the mechanic mode of operation.
+	 */
+	void SetSelectionMechanicMode(EUVEditorSelectionMode Mode, 
+		const FSelectionMechanicModeChangeOptions& Options = FSelectionMechanicModeChangeOptions());
+	
+	// Highlighting controls
+	
+	/**
+	 * Changes the visibility of the highlight.
+	 */
+	void SetHighlightVisible(bool bUnwrapHighlightVisible, bool bAppliedHighlightVisible, bool bRebuild = true);
+
+	struct FHighlightOptions
+	{
+		// When true, highlighting is based off of the preview meshes rather than the
+		// canonical meshes.
+		bool bBaseHighlightOnPreviews = false;
+
+		// When true, SetSelection calls will trigger a rebuild of the applied highlight
+		bool bAutoUpdateApplied = false;
+
+		// When true, SetSelection calls will trigger a rebuild of the unwrap highlight
+		bool bAutoUpdateUnwrap = false;
+
+		// If unwrap highlights are automatically updated, the start transform of these
+		// will be set to the centroid of the selection when the below is true. This
+		// requires centroid calculation but allows the highlight to be easily translated
+		// in response to gizmo movement.
+		bool bUseCentroidForUnwrapAutoUpdate = true;
+
+		// When building the unwrap highlight, show each edge and the paired edge to which 
+		// it can be welded (in different colors).
+		bool bShowPairedEdgeHighlights = false;
+	};
+	void SetHighlightOptions(const FHighlightOptions& Options);
+
+	void ClearHighlight(bool bClearForUnwrap = true, bool bClearForAppliedPreview = true);
+
+	/**
+	 * Build up a highlight of the current selection in the Unwrap, with a given transform
+	 * considered as its start transform.
+	 */
+	void RebuildUnwrapHighlight(const FTransform& StartTransform);
+
+	/**
+	 * Change the transform of the unwrap highlight without rebuilding it (for cheap movement
+	 * of the highlight when translating the elements). Note that unmoved paired edges (if 
+	 * paired edge highlighting is enabled) will still be rebuilt since their shape may be
+	 * changed by the movement of adjacent edges.
+	 */
+	void SetUnwrapHighlightTransform(const FTransform& NewTransform);
+
+	FTransform GetUnwrapHighlightTransform();
+
+	void RebuildAppliedPreviewHighlight();
+
+	/**
+	 * Broadcasted right before a selection change is applied (and therefore before a selection change
+	 * transaction is emitted). Useful if a user wants to emit their own bookend transaction first.
+	 * 
+	 * @param bEmitChangeAllowed If false, the callback must not emit any undo/redo transactions,
+	 *  likely because this is being called from an apply/revert of an existing transaction.
+	 */
+	DECLARE_MULTICAST_DELEGATE_OneParam(FOnPreSelectionChange, bool bEmitChangeAllowed);
+	FOnPreSelectionChange OnPreSelectionChange;
+
+	/**
+	 * Broadcasted after a selection change is applied (after the selection change transaction is
+	 * emitted, if relevant).
+	 * 
+	 * @param bEmitChangeAllowed If false, the callback must not emit any undo/redo transactions,
+	 *  likely because this is being called from an apply/revert of an existing transaction.
+	 */
+	DECLARE_MULTICAST_DELEGATE_OneParam(FOnSelectionChanged, bool bEmitChangeAllowed);
+	FOnSelectionChanged OnSelectionChanged;
+
+	/**
+	 * Broadcasted whenever the marquee rectangle is changed, since these changes don't trigger
+	 * normal selection broadcasts.
+	 */
+	FSimpleMulticastDelegate OnDragSelectionChanged;
+
+
+	// Undo/redo support:
+
+	/**
+	 * An object that can be as an undo/redo item, usually emitted by the selection API
+	 * itself. Expects UUVToolSelectionAPI to be the associated UObject.
+	 */
+	class FSelectionChange : public FToolCommandChange
+	{
+	public:
+		/**
+		 * When true, edge-type selections are not stored as eids, but rather as stable identifiers 
+		 * relative to the UnwrapCanonical mesh in the Target of the selection. This keeps them
+		 * from being invalidated when other transactions edit the mesh in a way that changes
+		 * Eids without actually changing topology.
+		 */
+		bool bUseStableUnwrapCanonicalIDsForEdges = true;
+
+		virtual void SetBefore(TArray<FUVToolSelection> Selections);
+		virtual void SetAfter(TArray<FUVToolSelection> Selections);
+
+		// Useful for storing a pending change and figuring out whether it needs to be
+		// emitted. Respects bUseStableUnwrapCanonicalIDsForEdges.
+		virtual TArray<FUVToolSelection> GetBefore() const;
+
+		virtual void Apply(UObject* Object) override;
+		virtual void Revert(UObject* Object) override;
+		virtual FString ToString() const override;
+
+	protected:
+		TArray<FUVToolSelection> Before;
+		TArray<FUVToolSelection> After;
+	};
+
+	/**
+	 * Preps a selection change transaction, if the user wants more control on what the previous and
+	 * current selection is.
+	 */
+	void BeginChange();
+
+	/**
+	 * Ends the active change and emits it via the EmitChangeAPI
+	 */
+	bool EndChangeAndEmitIfModified(bool bBroadcast);
+
+	// Initialization functions:
+	void Initialize(UInteractiveToolManager* ToolManagerIn, UWorld* UnwrapWorld,
+		UInputRouter* UnwrapInputRouterIn, UUVToolLivePreviewAPI* LivePreviewAPI,
+		UUVToolEmitChangeAPI* EmitChangeAPIIn);
+	// Should be called after Initialize()
+	void SetTargets(const TArray<TObjectPtr<UUVEditorToolMeshInput>>& TargetsIn);
+
+	// Called by the owner of the API
+	virtual void Render(IToolsContextRenderAPI* RenderAPI);
+	virtual void DrawHUD(FCanvas* Canvas, IToolsContextRenderAPI* RenderAPI);
+
+	// UUVToolContextObject
+	virtual void Shutdown() override;
+	virtual void OnToolEnded(UInteractiveTool* DeadTool) override;
+protected:
+	TArray<FUVToolSelection> CurrentSelections;
+
+	UPROPERTY()
+	TArray<TObjectPtr<UUVEditorToolMeshInput>> Targets;
+
+	TWeakObjectPtr<UInputRouter> UnwrapInputRouter = nullptr;
+
+	UPROPERTY()
+	TObjectPtr<UUVEditorMechanicAdapterTool> MechanicAdapter = nullptr;
+
+	UPROPERTY()
+	TObjectPtr<UUVToolSelectionHighlightMechanic> HighlightMechanic = nullptr;
+
+	UPROPERTY()
+	TObjectPtr<UUVEditorMeshSelectionMechanic> SelectionMechanic = nullptr;
+
+	UPROPERTY()
+	TObjectPtr<UUVToolEmitChangeAPI> EmitChangeAPI = nullptr;
+
+	TUniquePtr<FSelectionChange> PendingSelectionChange = nullptr;
+
+	FHighlightOptions HighlightOptions;
+
+	FVector3d CachedUnwrapSelectionCentroid;
+	bool bCachedUnwrapSelectionCentroidValid = false;
+};
+
+// UInterface for IUVToolSupportsSelection
+UINTERFACE(MinimalAPI)
+class UUVToolSupportsSelection : public UInterface
+{
+	GENERATED_BODY()
+};
+
+/**
+ * If a tool does not inherit from IUVToolSupportsSelection, then selection
+ * will automatically be cleared before the tool invocation via an undoable 
+ * transaction, to avoid a state where the selection refers to invalid items
+ * after tool completion.
+ * If a tool does inherit from IUVToolSupportsSelection, then the UV editor will
+ * not clear the selection before invocation, allowing the tool to use it. However,
+ * the tool it is expected to properly deal with selection itself and avoid an
+ * invalid state (including avoiding incorrect undo/redo selection event ordering.
+ */
+class IUVToolSupportsSelection
+{
+	GENERATED_BODY()
+};
diff --git a/Engine/Plugins/Experimental/UVEditor/Source/UVEditorTools/Public/Selection/UVToolSelectionHighlightMechanic.h b/Engine/Plugins/Experimental/UVEditor/Source/UVEditorTools/Public/Selection/UVToolSelectionHighlightMechanic.h
new file mode 100644
--- /dev/null
+++ b/Engine/Plugins/Experimental/UVEditor/Source/UVEditorTools/Public/Selection/UVToolSelectionHighlightMechanic.h
@@ -0,0 +1,112 @@
+// Copyright Epic Games, Inc. All Rights Reserved.
+
+#pragma once
+
+#include "CoreMinimal.h"
+
+#include "GeometryBase.h"
+#include "InteractionMechanic.h"
+
+#include "UVToolSelectionHighlightMechanic.generated.h"
+
+PREDECLARE_GEOMETRY(class FDynamicMesh3);
+PREDECLARE_GEOMETRY(class FUVToolSelection);
+
+class APreviewGeometryActor;
+class UMaterialInstanceDynamic;
+class UTriangleSetComponent;
+class ULineSetComponent;
+class UPointSetComponent;
+class UUVEditorToolMeshInput;
+class UWorld;
+
+/**
+ * Mechanic for highlighting mesh elements (usually selection) in the UV editor.
+ */
+UCLASS()
+class UVEDITORTOOLS_API UUVToolSelectionHighlightMechanic : public UInteractionMechanic
+{
+	GENERATED_BODY()
+
+	using FUVToolSelection = UE::Geometry::FUVToolSelection;
+
+public:
+
+	void Initialize(UWorld* UnwrapWorld, UWorld* LivePreviewWorld);
+	void Shutdown() override;
+
+	void SetIsVisible(bool bUnwrapHighlightVisible, bool bLivePreviewHighlightVisible);
+
+	/**
+	 * Rebuilds the unwrap mesh highlights.
+	 * 
+	 * @param StartTransform The transform to consider as the start transform, so that the highlights
+	 *   can be repositioned easily later via SetUnwrapHighlightTransform().
+	 * @param bUsePreviews If true, the selected element positions will be gotten from the UnwrapPreview
+	 *   of each selection target rather than UnwrapCanonical.
+	 */
+	void RebuildUnwrapHighlight(const TArray<FUVToolSelection>& Selections, 
+		const FTransform& StartTransform, bool bUsePreviews = false);
+
+	/**
+	 * Set the transform of the unwrap mesh highlights without rebuilding them, likely to 
+	 * follow the translation of the underlying elements via gizmo.
+	 * 
+	 * @param bRebuildStaticPairedEdges If true, and paired edge highlighting is enabled,
+	 *   unselected paired edges are rebuilt in the highlight. This is usually necessary
+	 *   because such edges may be attached to moving selected edges.
+	 * @param bUsePreviews If bRebuildStaticPairedEdges is true, this will prefer to use the
+	 *   previews rather than the unwrap canonical.
+	 */
+	void SetUnwrapHighlightTransform(const FTransform& Transform, bool bRebuildStaticPairedEdges = true, 
+		bool bUsePreviews = false);
+
+	/**
+	 * Get the current transform of the unwrap highlight.
+	 */
+	FTransform GetUnwrapHighlightTransform();
+
+	/**
+	 * Convert the unwrap selections to elements in the applied mesh, then highlight those there.
+	 * 
+	 * @param bUsePreviews If true, the previews are used rather than the canonicals both for converting
+	 *   over to applied mesh elements and to finding the highlight positions.
+	 */
+	void RebuildAppliedHighlightFromUnwrapSelection(const TArray<FUVToolSelection>& UnwrapSelections, bool bUsePreviews = false);
+
+	/**
+	 * When true, edge selections also highlight the edge pairings (i.e. edges that could weld with the
+	 * selected edges). Note that the call itself controls visibility and the next rebuild call, but 
+	 * won't cause an immediate rebuild on its own.
+	 */
+	void SetEnablePairedEdgeHighlights(bool bEnable);
+
+protected:
+
+	UPROPERTY()
+	TObjectPtr<APreviewGeometryActor> UnwrapGeometryActor = nullptr;
+	UPROPERTY()
+	TObjectPtr<UMaterialInstanceDynamic> TriangleSetMaterial = nullptr;
+
+	TWeakObjectPtr<UTriangleSetComponent> UnwrapTriangleSet = nullptr;
+	TWeakObjectPtr<ULineSetComponent> UnwrapLineSet = nullptr;
+	TWeakObjectPtr<ULineSetComponent> SewEdgePairingLineSet = nullptr;
+	TWeakObjectPtr<UPointSetComponent> UnwrapPointSet = nullptr;
+	TWeakObjectPtr<ULineSetComponent> UnwrapPairedEdgeLineSet = nullptr;
+
+	UPROPERTY()
+	TObjectPtr<APreviewGeometryActor> UnwrapStationaryGeometryActor = nullptr;
+	TWeakObjectPtr<ULineSetComponent> SewEdgeUnselectedPairingLineSet = nullptr;
+
+	// Per mesh, gives a list of endpoint vid pairs for the unselected edges that
+	// are part of edge pairs, for rebuilding when necessary.
+	TArray<TPair<TWeakObjectPtr<UUVEditorToolMeshInput>,
+		TArray<TPair<int32, int32>>>> StaticPairedEdgeVidsPerMesh;
+
+	UPROPERTY()
+	TObjectPtr<APreviewGeometryActor> LivePreviewGeometryActor = nullptr;
+	TWeakObjectPtr<ULineSetComponent> LivePreviewLineSet = nullptr;
+	TWeakObjectPtr<UPointSetComponent> LivePreviewPointSet = nullptr;
+
+	bool bPairedEdgeHighlightsEnabled = true;
+};
diff --git a/Engine/Plugins/Experimental/UVEditor/Source/UVEditorTools/Public/UVEditorMechanicAdapterTool.h b/Engine/Plugins/Experimental/UVEditor/Source/UVEditorTools/Public/UVEditorMechanicAdapterTool.h
new file mode 100644
--- /dev/null
+++ b/Engine/Plugins/Experimental/UVEditor/Source/UVEditorTools/Public/UVEditorMechanicAdapterTool.h
@@ -0,0 +1,24 @@
+// Copyright Epic Games, Inc. All Rights Reserved.
+
+#pragma once
+
+#include "CoreMinimal.h"
+
+#include "UVEditorMechanicAdapterTool.generated.h"
+
+/**
+ * This is a dummy tool that only exists to make it possible for mechanics to be used in
+ * outside of an actual tool. The tool can be passed in to the mechanics' Setup() calls. 
+ * If necessary, the behaviors can be extracted via GetInputBehaviors() to be placed into 
+ * an input router.
+ */
+UCLASS()
+class UVEDITORTOOLS_API UUVEditorMechanicAdapterTool : public UInteractiveTool
+{
+	GENERATED_BODY()
+
+public:
+	UInteractiveToolManager* ToolManager = nullptr;
+
+	virtual UInteractiveToolManager* GetToolManager() const override { return ToolManager; }
+};
\ No newline at end of file
diff --git a/Engine/Plugins/Experimental/UVEditor/Source/UVEditorTools/Public/UVIslandConformalUnwrapAction.h b/Engine/Plugins/Experimental/UVEditor/Source/UVEditorTools/Public/UVIslandConformalUnwrapAction.h
deleted file mode 100644
diff --git a/Engine/Plugins/Experimental/UVEditor/Source/UVEditorTools/Public/UVSeamSewAction.h b/Engine/Plugins/Experimental/UVEditor/Source/UVEditorTools/Public/UVSeamSewAction.h
deleted file mode 100644
diff --git a/Engine/Plugins/Experimental/UVEditor/Source/UVEditorTools/Public/UVSelectTool.h b/Engine/Plugins/Experimental/UVEditor/Source/UVEditorTools/Public/UVSelectTool.h
--- a/Engine/Plugins/Experimental/UVEditor/Source/UVEditorTools/Public/UVSelectTool.h
+++ b/Engine/Plugins/Experimental/UVEditor/Source/UVEditorTools/Public/UVSelectTool.h
@@ -4,34 +4,18 @@
 
 #include "CoreMinimal.h"
 
-#include "DynamicMesh/DynamicMeshAABBTree3.h"
 #include "FrameTypes.h"
 #include "GeometryBase.h"
 #include "InteractiveTool.h"
 #include "InteractiveToolBuilder.h"
-#include "TargetInterfaces/UVUnwrapDynamicMesh.h"
-#include "UVToolContextObjects.h"
-#include "UVEditorToolAnalyticsUtils.h"
-#include "IndexTypes.h"
+#include "Selection/UVToolSelection.h"
+#include "Selection/UVToolSelectionAPI.h" // IUVToolSupportsSelection
 
 #include "UVSelectTool.generated.h"
 
-PREDECLARE_GEOMETRY(class FUVEditorDynamicMeshSelection);
-class APreviewGeometryActor;
-class FToolCommandChange;
-class ULineSetComponent;
-class UUVEditorMeshSelectionMechanic;
-class UPointSetComponent;
-class UToolTargetManager;
 class UCombinedTransformGizmo;
 class UTransformProxy;
-class UUVEditorMode;
-class UUVSelectToolChangeRouter;
-class UUVToolStateObjectStore;
-class UPreviewMesh;
 class UUVEditorToolMeshInput;
-class UUVSeamSewAction;
-class UUVIslandConformalUnwrapAction;
 class UUVToolEmitChangeAPI;
 class UUVToolViewportButtonsAPI;
 
@@ -49,61 +33,17 @@ public:
 	const TArray<TObjectPtr<UUVEditorToolMeshInput>>* Targets = nullptr;
 };
 
-UENUM()
-enum class ESelectToolAction
-{
-	NoAction,
-
-	Sew,
-	Split,
-	IslandConformalUnwrap,
-};
-
-UCLASS()
-class UVEDITORTOOLS_API USelectToolActionPropertySet : public UInteractiveToolPropertySet
-{
-	GENERATED_BODY()
-
-public:
-	TWeakObjectPtr<UUVSelectTool> ParentTool;
-
-	void Initialize(UUVSelectTool* ParentToolIn) { ParentTool = ParentToolIn; }
-	void PostAction(ESelectToolAction Action);
-
-	//~ NOTE: the below comments don't actually get used as tooltips because we use a detail
-	//~ customization, and the tooltips are specified in UVEditorCommands.cpp. Change them
-	//~ there if needed!
-
-	/** Sew edges. The red edges will be sewn to the green edges */
-	UFUNCTION(CallInEditor, Category = Actions)
-	void Sew() { PostAction(ESelectToolAction::Sew); }
-
-	/** Given an edge selection, split those edges. Given a vertex selection, split any selected bowtie verts. */
-	UFUNCTION(CallInEditor, Category = Actions)
-	void Split() { PostAction(ESelectToolAction::Split); };
-
-	/* Apply a conformal unwrap to the selected UV islands */
-	UFUNCTION(CallInEditor, Category = Actions)
-	void IslandConformalUnwrap();
-};
-
-
-
 /**
- * A tool for selecting elements of a flat FDynamicMesh corresponding to a UV layer of some asset.
- * If bGizmoEnabled is set to true, the selected elements can be moved around.
- *
- * TODO: Doesn't have undo/redo. Will get broken up into pieces later, probably.
+ * The tool in the UV editor that secretly runs when other tools are not running. It uses the
+ * selection api to allow the user to select elements, and has a gizmo that can be used to
+ * transform these elements.
  */
 UCLASS()
-class UVEDITORTOOLS_API UUVSelectTool : public UInteractiveTool
+class UVEDITORTOOLS_API UUVSelectTool : public UInteractiveTool, public IUVToolSupportsSelection
 {
 	GENERATED_BODY()
 
-	using FDynamicMeshAABBTree3 = UE::Geometry::FDynamicMeshAABBTree3;
-
 public:
-	virtual void SetWorld(UWorld* World) { TargetWorld = World; }
 	
 	/**
 	 * The tool will operate on the meshes given here.
@@ -113,8 +53,8 @@ public:
 		Targets = TargetsIn;
 	}
 
-	// Used by undo/redo changes to update the tool state
-	void SetSelection(const UE::Geometry::FUVEditorDynamicMeshSelection& NewSelection);
+	// Used by undo/redo.
+	FTransform GetGizmoTransform() const;
 	void SetGizmoTransform(const FTransform& NewTransform);
 
 	// UInteractiveTool
@@ -127,122 +67,46 @@ public:
 	virtual bool HasCancel() const override { return false; }
 	virtual bool HasAccept() const override { return false; }
 
-	void RequestAction(ESelectToolAction ActionType);
-
 protected:
-	virtual void OnSelectionChanged();
-	virtual void ClearWarning();
+	virtual void OnSelectionChanged(bool bEmitChangeAllowed);
 
 	// Callbacks we'll receive from the gizmo proxy
 	virtual void GizmoTransformChanged(UTransformProxy* Proxy, FTransform Transform);
 	virtual void GizmoTransformStarted(UTransformProxy* Proxy);
 	virtual void GizmoTransformEnded(UTransformProxy* Proxy);
 
-
 	virtual void ApplyGizmoTransform();
 	virtual void UpdateGizmo();
-	virtual void UpdateLivePreviewLines();
-
-	void UpdateSelectionMode();	
-
-	UWorld* TargetWorld;
 
 	UPROPERTY()
 	TArray<TObjectPtr<UUVEditorToolMeshInput>> Targets;
 
-	UPROPERTY()
-	TObjectPtr<USelectToolActionPropertySet> ToolActions = nullptr;
-
 	UPROPERTY()
 	TObjectPtr<UUVToolViewportButtonsAPI> ViewportButtonsAPI = nullptr;
-
-	UPROPERTY()
-	TObjectPtr<UUVEditorMeshSelectionMechanic> SelectionMechanic = nullptr;
-
-	UPROPERTY()
-	TObjectPtr<UCombinedTransformGizmo> TransformGizmo = nullptr;
-
-	TArray<TSharedPtr<FDynamicMeshAABBTree3>> AABBTrees;
-
-	UPROPERTY()
-	TObjectPtr<APreviewGeometryActor> LivePreviewGeometryActor = nullptr;
-
-	UPROPERTY()
-	TObjectPtr<ULineSetComponent> LivePreviewLineSet = nullptr;
-
-	UPROPERTY()
-	TObjectPtr<UPointSetComponent> LivePreviewPointSet = nullptr;
-
+	
 	UPROPERTY()
 	TObjectPtr<UUVToolEmitChangeAPI> EmitChangeAPI = nullptr;
 
 	UPROPERTY()
-	TObjectPtr<UUVSelectToolChangeRouter> ChangeRouter = nullptr;
+	TObjectPtr<UUVToolSelectionAPI> SelectionAPI = nullptr;
 
 	UPROPERTY()
-	TObjectPtr<UUVSeamSewAction> SewAction = nullptr;
+	TObjectPtr<UUVEditorMeshSelectionMechanic> SelectionMechanic = nullptr;
 
 	UPROPERTY()
-	TObjectPtr<UUVIslandConformalUnwrapAction> IslandConformalUnwrapAction = nullptr;
+	TObjectPtr<UCombinedTransformGizmo> TransformGizmo = nullptr;
 
 	UE::Geometry::FFrame3d InitialGizmoFrame;
 	FTransform UnappliedGizmoTransform;
 	bool bInDrag = false;
-	bool bIgnoreOnCanonicalChange = false;
+	bool bUpdateGizmoOnCanonicalChange = true;
 	bool bGizmoTransformNeedsApplication = false;
 
-	TArray<int32> MovingVids;
-	TArray<int32> SelectedVids;
-	TArray<int32> SelectedTids;
-	TArray<FVector3d> MovingVertOriginalPositions;
-	int32 SelectionTargetIndex;
-	TSet<int32> LivePreviewEids;
-	TSet<int32> LivePreviewVids;
-	
-	// When selecting edges, used to hold the edges as pairs of vids, because the eids change
-	// during undo/redo and other topological operations.
-	TArray<UE::Geometry::FIndex2i> CurrentSelectionVidPairs;
-
-	ESelectToolAction PendingAction;
-
-	void ApplyAction(ESelectToolAction ActionType);
-	void ApplySplit();
-	void ApplySplitEdges();
-	void ApplySplitBowtieVertices();
-
-
-
-	//
-	// Analytics
-	//
-
-	struct FActionHistoryItem
-	{
-		FDateTime Timestamp;
-		ESelectToolAction ActionType;
-		
-		// if ActionType == IslandConformalUnwrap then NumOperands is #triangles in island selection
-		// if ActionType == Split                 then NumOperands is #edges in selection
-		// if ActionType == Sew                   then NumOperands is #edges in selection
-		int32 NumOperands = -1;
-	};
-	
-	TArray<FActionHistoryItem> AnalyticsActionHistory;
-	UE::Geometry::UVEditorAnalytics::FTargetAnalytics InputTargetAnalytics;
-	FDateTime ToolStartTimeAnalytics;
-	void RecordAnalytics();
+	TArray<UE::Geometry::FUVToolSelection> CurrentSelections;
+	// The outer arrays are 1:1 with the Selections array obtained from the Selection API
+	TArray<TArray<int32>> RenderUpdateTidsPerSelection;
+	// Inner arrays for these two are 1:1 with each other
+	TArray<TArray<int32>> MovingVidsPerSelection;
+	TArray<TArray<FVector3d>> MovingVertOriginalPositionsPerSelection;
 };
 
-/**
- * A helper context object that we can use as the target of undo/redo events to apply them
- * to the current invocation of the select tool (which may have different gizmo/selection
- * pointers than those that were around when the change was emitted).
- */
-UCLASS()
-class UVEDITORTOOLS_API UUVSelectToolChangeRouter : public UUVToolContextObject
-{
-	GENERATED_BODY()
-public:
-
-	TWeakObjectPtr<UUVSelectTool> CurrentSelectTool = nullptr;
-};
diff --git a/Engine/Plugins/Experimental/UVEditor/Source/UVEditorTools/Public/UVToolAction.h b/Engine/Plugins/Experimental/UVEditor/Source/UVEditorTools/Public/UVToolAction.h
deleted file mode 100644
diff --git a/Engine/Plugins/Experimental/UVEditor/Source/UVEditorTools/Public/UVToolContextObjects.h b/Engine/Plugins/Experimental/UVEditor/Source/UVEditorTools/Public/UVToolContextObjects.h
deleted file mode 100644
diff --git a/Engine/Plugins/Experimental/UVEditor/Source/UVEditorToolsEditorOnly/Private/UVEditorParameterizeMeshTool.cpp b/Engine/Plugins/Experimental/UVEditor/Source/UVEditorToolsEditorOnly/Private/UVEditorParameterizeMeshTool.cpp
--- a/Engine/Plugins/Experimental/UVEditor/Source/UVEditorToolsEditorOnly/Private/UVEditorParameterizeMeshTool.cpp
+++ b/Engine/Plugins/Experimental/UVEditor/Source/UVEditorToolsEditorOnly/Private/UVEditorParameterizeMeshTool.cpp
@@ -7,7 +7,7 @@
 #include "ParameterizationOps/ParameterizeMeshOp.h"
 #include "Properties/ParameterizeMeshProperties.h"
 #include "ToolTargets/UVEditorToolMeshInput.h"
-#include "UVToolContextObjects.h"
+#include "ContextObjects/UVToolContextObjects.h"
 #include "ContextObjectStore.h"
 #include "EngineAnalytics.h"
 
