UVEditor: Non-unity build fixes.

#rnx

diff --git a/Engine/Plugins/Experimental/UVEditor/Source/UVEditorTools/Public/ContextObjects/UVToolAssetInputsContext.h b/Engine/Plugins/Experimental/UVEditor/Source/UVEditorTools/Public/ContextObjects/UVToolAssetInputsContext.h
--- a/Engine/Plugins/Experimental/UVEditor/Source/UVEditorTools/Public/ContextObjects/UVToolAssetInputsContext.h
+++ b/Engine/Plugins/Experimental/UVEditor/Source/UVEditorTools/Public/ContextObjects/UVToolAssetInputsContext.h
@@ -4,6 +4,8 @@
 
 #include "CoreMinimal.h"
 
+#include "ContextObjects/UVToolContextObjects.h"
+
 #include "UVToolAssetInputsContext.generated.h"
 
 /**
diff --git a/Engine/Plugins/Experimental/UVEditor/Source/UVEditorTools/Public/UVEditorMechanicAdapterTool.h b/Engine/Plugins/Experimental/UVEditor/Source/UVEditorTools/Public/UVEditorMechanicAdapterTool.h
--- a/Engine/Plugins/Experimental/UVEditor/Source/UVEditorTools/Public/UVEditorMechanicAdapterTool.h
+++ b/Engine/Plugins/Experimental/UVEditor/Source/UVEditorTools/Public/UVEditorMechanicAdapterTool.h
@@ -4,8 +4,12 @@
 
 #include "CoreMinimal.h"
 
+#include "InteractiveTool.h"
+
 #include "UVEditorMechanicAdapterTool.generated.h"
 
+class UInteractiveToolManager;
+
 /**
  * This is a dummy tool that only exists to make it possible for mechanics to be used in
  * outside of an actual tool. The tool can be passed in to the mechanics' Setup() calls. 
