UVEditor: Implemented initial support for UDIM aware layout in the UVEditor, allowing islands to remain within a tile if they were already there. Additionally changed the UDIM display from asset to use current island positions rather than the UDIM polygroup.

diff --git a/Engine/Plugins/Experimental/UVEditor/Source/UVEditor/Private/UVEditorMode.cpp b/Engine/Plugins/Experimental/UVEditor/Source/UVEditor/Private/UVEditorMode.cpp
--- a/Engine/Plugins/Experimental/UVEditor/Source/UVEditor/Private/UVEditorMode.cpp
+++ b/Engine/Plugins/Experimental/UVEditor/Source/UVEditor/Private/UVEditorMode.cpp
@@ -46,6 +46,8 @@
 #include "UVEditorUXSettings.h"
 #include "UDIMUtilities.h"
 #include "EditorSupportDelegates.h"
+#include "Utilities/MeshUDIMClassifier.h"
+#include "UVEditorLogging.h"
 
 #define LOCTEXT_NAMESPACE "UUVEditorMode"
 
@@ -991,32 +993,26 @@ void UUVEditorMode::PopulateUDIMsByAsset(int32 AssetId, TArray<FUDIMSpecifier>&
 	{
 		return;
 	}
-	if (ToolInputObjects[AssetId]->AppliedCanonical->HasAttributes())
+
+	if (ToolInputObjects[AssetId]->AppliedCanonical->HasAttributes() && ToolInputObjects[AssetId]->UVLayerIndex >= 0)
 	{
-		for (int32 GroupLayerIndex = 0; GroupLayerIndex < ToolInputObjects[AssetId]->AppliedCanonical->Attributes()->NumPolygroupLayers(); ++GroupLayerIndex)
+		FDynamicMeshUVOverlay* UVLayer = ToolInputObjects[AssetId]->AppliedCanonical->Attributes()->GetUVLayer(ToolInputObjects[AssetId]->UVLayerIndex);
+		FDynamicMeshUDIMClassifier TileClassifier(UVLayer);
+		TArray<FVector2i> Tiles = TileClassifier.ActiveTiles();
+		for (FVector2i Tile : Tiles)
 		{
-			if (ToolInputObjects[AssetId]->AppliedCanonical->Attributes()->GetPolygroupLayer(GroupLayerIndex)->GetName() == "UDIM")
+			if (Tile.X >= 10 || Tile.X < 0 || Tile.Y < 0)
 			{
-				FDynamicMeshPolygroupAttribute* UDIMPolygroup = ToolInputObjects[AssetId]->AppliedCanonical->Attributes()->GetPolygroupLayer(GroupLayerIndex);
-
-				// This isn't very efficient, as a linear pass over all Tids. But there doesn't seem to be a better way to list all the groups present.
-				for (int32 Tid : ToolInputObjects[AssetId]->AppliedCanonical->TriangleIndicesItr())
-				{
-					int32 UDIMValue = UDIMPolygroup->GetValue(Tid);
-
-					// Discard any invalid UDIM group labels
-					if (UDIMValue < 1001)
-					{
-						continue;
-					}
-					FUDIMSpecifier UDIMSpecifier;
-					UDIMSpecifier.UDIM = UDIMValue;
-					UE::TextureUtilitiesCommon::ExtractUDIMCoordinates(UDIMValue, UDIMSpecifier.UCoord, UDIMSpecifier.VCoord);
-					UDIMsOut.AddUnique(UDIMSpecifier);
-				}
+				UE_LOG(LogUVEditor, Warning, TEXT("Tile <%d,%d> is out of bounds of the UDIM10 convention, skipping..."), Tile.X, Tile.Y);
+			}
+			else
+			{
+				FUDIMSpecifier UDIMSpecifier;
+				UDIMSpecifier.UCoord = Tile.X;
+				UDIMSpecifier.VCoord = Tile.Y;
+				UDIMSpecifier.UDIM = UE::TextureUtilitiesCommon::GetUDIMIndex(Tile.X, Tile.Y);
+				UDIMsOut.AddUnique(UDIMSpecifier);
 			}
-
-			break; // Skip looking at other polygroup layers.
 		}
 	}
 }
diff --git a/Engine/Plugins/Experimental/UVEditor/Source/UVEditorTools/Private/Operators/UVEditorUVLayoutOp.cpp b/Engine/Plugins/Experimental/UVEditor/Source/UVEditorTools/Private/Operators/UVEditorUVLayoutOp.cpp
new file mode 100644
--- /dev/null
+++ b/Engine/Plugins/Experimental/UVEditor/Source/UVEditorTools/Private/Operators/UVEditorUVLayoutOp.cpp
@@ -0,0 +1,168 @@
+// Copyright Epic Games, Inc. All Rights Reserved.
+
+#include "Operators/UVEditorUVLayoutOp.h"
+
+#include "DynamicMesh/DynamicMesh3.h"
+#include "DynamicMesh/DynamicMeshAttributeSet.h"
+#include "Parameterization/MeshUVPacking.h"
+#include "Selections/MeshConnectedComponents.h"
+#include "Utilities/MeshUDIMClassifier.h"
+#include "UVEditorUXSettings.h"
+
+using namespace UE::Geometry;
+
+void FUVEditorUVLayoutOp::SetTransform(const FTransformSRT3d& Transform)
+{
+	ResultTransform = Transform;
+}
+
+
+void FUVEditorUVLayoutOp::CalculateResult(FProgressCancel* Progress)
+{
+	TRACE_CPUPROFILER_EVENT_SCOPE(UVEditorUVLayoutOp_CalculateResult);
+	
+	if (Progress && Progress->Cancelled())
+	{
+		return;
+	}
+	ResultMesh->Copy(*OriginalMesh, true, true, true, true);
+
+	if (!ensureMsgf(ResultMesh->HasAttributes(), TEXT("Attributes not found on mesh? Conversion should always create them, so this operator should not need to do so.")))
+	{
+		ResultMesh->EnableAttributes();
+	}
+
+	if (Progress && Progress->Cancelled())
+	{
+		return;
+	}
+
+
+	int UVLayerInput = UVLayerIndex;
+	FDynamicMeshUVOverlay* UseUVLayer = ResultMesh->Attributes()->GetUVLayer(UVLayerInput);
+	
+
+	bool bWillRepackIslands = (UVLayoutMode != EUVEditorUVLayoutOpLayoutModes::TransformOnly);
+
+	// split bowties so that we can process islands independently
+	if (bWillRepackIslands || bAlwaysSplitBowties)
+	{
+		UseUVLayer->SplitBowties();
+	}
+
+	if (Progress && Progress->Cancelled())
+	{
+		return;
+	}
+
+	if (bMaintainOriginatingUDIM)
+	{
+		FDynamicMeshUDIMClassifier TileClassifier(UseUVLayer);
+		TArray<FVector2i> Tiles = TileClassifier.ActiveTiles();
+
+		for (FVector2i TileIndex : Tiles)
+		{
+			TUniquePtr<TArray<int32>> TileTids = MakeUnique<TArray<int32>>(TileClassifier.TidsForTile(TileIndex));
+
+			// Do this first, so we don't need to keep the TileTids around after moving it into the packer.
+			TSet<int32> ElementsToMove;
+			ElementsToMove.Reserve(TileTids->Num() * 3);
+			for (int Tid : *TileTids)
+			{
+				FIndex3i Elements = UseUVLayer->GetTriangle(Tid);
+				ElementsToMove.Add(Elements[0]);
+				ElementsToMove.Add(Elements[1]);
+				ElementsToMove.Add(Elements[2]);
+			}
+
+			// TODO: There is a second connected components call inside the packer that might be unnessessary. Could be a future optimization.
+			FDynamicMeshUVPacker Packer(UseUVLayer, MoveTemp(TileTids) );
+			ExecutePacker(Packer);
+			if (Progress && Progress->Cancelled())
+			{
+				return;
+			}
+
+			for (int32 Element : ElementsToMove)
+			{
+				FVector2f UV = FUVEditorUXSettings::InternalUVToExternalUV(UseUVLayer->GetElement(Element));
+				UV = (UV) + (FVector2f)(TileIndex);
+				UseUVLayer->SetElement(Element, FUVEditorUXSettings::ExternalUVToInternalUV(UV));
+			}			
+		}
+	}
+	else
+	{
+		FDynamicMeshUVPacker Packer(UseUVLayer);
+		ExecutePacker(Packer);
+		if (Progress && Progress->Cancelled())
+		{
+			return;
+		}
+	}
+
+	if (UVScaleFactor != 1.0 || UVTranslation != FVector2f::Zero() )
+	{
+		for (int ElementID : UseUVLayer->ElementIndicesItr())
+		{
+			FVector2f UV = FUVEditorUXSettings::InternalUVToExternalUV(UseUVLayer->GetElement(ElementID));
+			UV = (UV * UVScaleFactor) + UVTranslation;
+			UseUVLayer->SetElement(ElementID, FUVEditorUXSettings::ExternalUVToInternalUV(UV));
+		}
+	}
+
+}
+
+void FUVEditorUVLayoutOp::ExecutePacker(FDynamicMeshUVPacker& Packer)
+{
+	Packer.TextureResolution = this->TextureResolution;
+	Packer.GutterSize = this->GutterSize;
+	Packer.bAllowFlips = this->bAllowFlips;
+
+	if (UVLayoutMode == EUVEditorUVLayoutOpLayoutModes::RepackToUnitRect)
+	{
+		if (Packer.StandardPack() == false)
+		{
+			// failed... what to do?
+			return;
+		}
+	}
+	else if (UVLayoutMode == EUVEditorUVLayoutOpLayoutModes::StackInUnitRect)
+	{
+		if (Packer.StackPack() == false)
+		{
+			// failed... what to do?
+			return;
+		}
+	}
+}
+
+TUniquePtr<FDynamicMeshOperator> UUVEditorUVLayoutOperatorFactory::MakeNewOperator()
+{
+	TUniquePtr<FUVEditorUVLayoutOp> Op = MakeUnique<FUVEditorUVLayoutOp>();
+
+	Op->OriginalMesh = OriginalMesh;
+
+	switch (Settings->LayoutType)
+	{
+	case EUVEditorUVLayoutType::Transform:
+		Op->UVLayoutMode = EUVEditorUVLayoutOpLayoutModes::TransformOnly;
+		break;
+	case EUVEditorUVLayoutType::Stack:
+		Op->UVLayoutMode = EUVEditorUVLayoutOpLayoutModes::StackInUnitRect;
+		break;
+	case EUVEditorUVLayoutType::Repack:
+		Op->UVLayoutMode = EUVEditorUVLayoutOpLayoutModes::RepackToUnitRect;
+		break;
+	}
+
+	Op->UVLayerIndex = GetSelectedUVChannel();
+	Op->TextureResolution = Settings->TextureResolution;
+	Op->bAllowFlips = Settings->bAllowFlips;
+	Op->UVScaleFactor = Settings->Scale;
+	Op->UVTranslation = FVector2f(Settings->Translation);
+	Op->SetTransform(TargetTransform);
+	Op->bMaintainOriginatingUDIM = Settings->bEnableUDIMLayout;
+
+	return Op;
+}
\ No newline at end of file
diff --git a/Engine/Plugins/Experimental/UVEditor/Source/UVEditorTools/Private/UVEditorLayoutTool.cpp b/Engine/Plugins/Experimental/UVEditor/Source/UVEditorTools/Private/UVEditorLayoutTool.cpp
--- a/Engine/Plugins/Experimental/UVEditor/Source/UVEditorTools/Private/UVEditorLayoutTool.cpp
+++ b/Engine/Plugins/Experimental/UVEditor/Source/UVEditorTools/Private/UVEditorLayoutTool.cpp
@@ -7,12 +7,12 @@
 #include "DynamicMesh/DynamicMeshChangeTracker.h"
 #include "InteractiveToolManager.h"
 #include "MeshOpPreviewHelpers.h" // UMeshOpPreviewWithBackgroundCompute
-#include "ParameterizationOps/UVLayoutOp.h"
-#include "Properties/UVLayoutProperties.h"
+#include "Operators/UVEditorUVLayoutOp.h"
 #include "ToolTargets/UVEditorToolMeshInput.h"
 #include "ContextObjects/UVToolContextObjects.h"
 #include "EngineAnalytics.h"
 #include "UVEditorToolAnalyticsUtils.h"
+#include "UVEditorUXSettings.h"
 
 using namespace UE::Geometry;
 
@@ -43,15 +43,17 @@ void UUVEditorLayoutTool::Setup()
 
 	UInteractiveTool::Setup();
 
-	Settings = NewObject<UUVLayoutProperties>(this);
+	Settings = NewObject<UUVEditorUVLayoutProperties>(this);
 	Settings->RestoreProperties(this);
+	Settings->bUDIMCVAREnabled = (FUVEditorUXSettings::CVarEnablePrototypeUDIMSupport.GetValueOnGameThread() > 0);
 	AddToolPropertySource(Settings);
 
+
 	Factories.SetNum(Targets.Num());
 	for (int32 TargetIndex = 0; TargetIndex < Targets.Num(); ++TargetIndex)
 	{
 		TObjectPtr<UUVEditorToolMeshInput> Target = Targets[TargetIndex];
-		Factories[TargetIndex] = NewObject<UUVLayoutOperatorFactory>();
+		Factories[TargetIndex] = NewObject<UUVEditorUVLayoutOperatorFactory>();
 		Factories[TargetIndex]->TargetTransform = Target->AppliedPreview->PreviewMesh->GetTransform();
 		Factories[TargetIndex]->Settings = Settings;
 		Factories[TargetIndex]->OriginalMesh = Target->AppliedCanonical;
diff --git a/Engine/Plugins/Experimental/UVEditor/Source/UVEditorTools/Private/UVEditorLogging.cpp b/Engine/Plugins/Experimental/UVEditor/Source/UVEditorTools/Private/UVEditorLogging.cpp
new file mode 100644
--- /dev/null
+++ b/Engine/Plugins/Experimental/UVEditor/Source/UVEditorTools/Private/UVEditorLogging.cpp
@@ -0,0 +1,5 @@
+// Copyright Epic Games, Inc. All Rights Reserved.
+
+#include "UVEditorLogging.h"
+
+DEFINE_LOG_CATEGORY(LogUVEditor);
diff --git a/Engine/Plugins/Experimental/UVEditor/Source/UVEditorTools/Private/Utilities/MeshUDIMClassifier.cpp b/Engine/Plugins/Experimental/UVEditor/Source/UVEditorTools/Private/Utilities/MeshUDIMClassifier.cpp
new file mode 100644
--- /dev/null
+++ b/Engine/Plugins/Experimental/UVEditor/Source/UVEditorTools/Private/Utilities/MeshUDIMClassifier.cpp
@@ -0,0 +1,89 @@
+// Copyright Epic Games, Inc. All Rights Reserved.
+
+#include "Utilities/MeshUDIMClassifier.h"
+#include "Spatial/SparseDynamicOctree3.h"
+#include "UVEditorUXSettings.h"
+#include "Selections/MeshConnectedComponents.h"
+
+using namespace UE::Geometry;
+
+FDynamicMeshUDIMClassifier::FDynamicMeshUDIMClassifier(FDynamicMeshUVOverlay* UVOverlayIn)
+{
+	UVOverlay = UVOverlayIn;
+	ClassifyUDIMs();
+}
+
+TArray<FVector2i> FDynamicMeshUDIMClassifier::ActiveTiles() const
+{
+	TArray<FVector2i> Keys;
+	UDIMs.GetKeys(Keys);
+	return Keys;
+}
+
+TArray<int32> FDynamicMeshUDIMClassifier::TidsForTile(FVector2i TileIndexIn) const
+{	
+	if (ensure(UDIMs.Contains(TileIndexIn)))
+	{
+		return *(UDIMs.Find(TileIndexIn));
+	}
+	else
+	{
+		return TArray<int32>();
+	}
+}
+
+void FDynamicMeshUDIMClassifier::ClassifyUDIMs()
+{
+	int32 CurrentUDIMIndex = 0;
+	FDynamicMesh3* Mesh = UVOverlay->GetParentMesh();
+
+	auto UVIslandPredicate = [this](int32 Triangle0, int32 Triangle1)
+	{
+		return UVOverlay->AreTrianglesConnected(Triangle0, Triangle1);
+	};
+
+	FMeshConnectedComponents UVComponents(Mesh);
+	UVComponents.FindConnectedTriangles(UVIslandPredicate);
+
+	auto UVTriangleToUDIM = [this](int32 Tid)
+	{
+		FVector2i UDIM;
+		FVector2f Vertex0, Vertex1, Vertex2;
+		UVOverlay->GetTriElements(Tid, Vertex0, Vertex1, Vertex2);
+		FVector2f BaryCenter = (Vertex0 + Vertex1 + Vertex2) / 3.0f;
+		BaryCenter = FUVEditorUXSettings::InternalUVToExternalUV(BaryCenter);
+		UDIM.X = FMath::Floor(BaryCenter.X);
+		UDIM.Y = FMath::Floor(BaryCenter.Y);
+		return UDIM;
+	};
+
+	auto IslandToUDIM = [this, &UVComponents, &UVTriangleToUDIM](int32 Cid)
+	{
+		TMap<FVector2i, int32> UDIMCount;
+		for (int32 Tid : UVComponents.GetComponent(Cid).Indices)
+		{
+			FVector2i UDIM = UVTriangleToUDIM(Tid);
+			UDIMCount.FindOrAdd(UDIM, 0)++;
+		}
+		FVector2i MaximumUDIM;
+		int32 MaximumUDIMCount = 0;
+		for (auto UDIMEntry : UDIMCount)
+		{
+			if (UDIMEntry.Value > MaximumUDIMCount)
+			{
+				MaximumUDIMCount = UDIMEntry.Value;
+				MaximumUDIM = UDIMEntry.Key;
+			}
+		}
+		return MaximumUDIM;
+	};
+
+	for (int32 Cid = 0; Cid < UVComponents.Num(); ++Cid)
+	{
+		FVector2i UDIM = IslandToUDIM(Cid);
+		UDIMs.FindOrAdd(UDIM).Append(UVComponents.GetComponent(Cid).Indices);
+	}
+
+	return;
+
+}
diff --git a/Engine/Plugins/Experimental/UVEditor/Source/UVEditorTools/Public/Operators/UVEditorUVLayoutOp.h b/Engine/Plugins/Experimental/UVEditor/Source/UVEditorTools/Public/Operators/UVEditorUVLayoutOp.h
new file mode 100644
--- /dev/null
+++ b/Engine/Plugins/Experimental/UVEditor/Source/UVEditorTools/Public/Operators/UVEditorUVLayoutOp.h
@@ -0,0 +1,140 @@
+// Copyright Epic Games, Inc. All Rights Reserved.
+
+#pragma once
+
+#include "CoreMinimal.h"
+#include "Util/ProgressCancel.h"
+#include "ModelingOperators.h"
+#include "InteractiveTool.h"
+
+#include "UVEditorUVLayoutOp.generated.h"
+
+class UUVLayoutProperties;
+
+/**
+ * UV Layout Strategies for the UV Layout Tool
+ */
+UENUM()
+enum class EUVEditorUVLayoutType
+{
+	/** Apply Scale and Translation properties to all UV values */
+	Transform,
+	/** Uniformly scale and translate each UV island individually to pack it into the unit square, i.e. fit between 0 and 1 with overlap */
+	Stack,
+	/** Uniformly scale and translate UV islands collectively to pack them into the unit square, i.e. fit between 0 and 1 with no overlap */
+	Repack
+};
+
+
+/**
+ * UV Layout Settings
+ */
+UCLASS()
+class UVEDITORTOOLS_API UUVEditorUVLayoutProperties : public UInteractiveToolPropertySet
+{
+	GENERATED_BODY()
+
+public:
+	/** Type of layout applied to input UVs */
+	UPROPERTY(EditAnywhere, Category = "UV Layout")
+	EUVEditorUVLayoutType LayoutType = EUVEditorUVLayoutType::Repack;
+
+	/** Expected resolution of the output textures; this controls spacing left between UV islands to avoid interpolation artifacts */
+	UPROPERTY(EditAnywhere, Category = "UV Layout", meta = (UIMin = "64", UIMax = "2048", ClampMin = "2", ClampMax = "4096"))
+	int TextureResolution = 1024;
+
+	/** Uniform scale applied to UVs after packing */
+	UPROPERTY(EditAnywhere, Category = "UV Layout", meta = (UIMin = "0.1", UIMax = "5.0", ClampMin = "0.0001", ClampMax = "10000"))
+	float Scale = 1;
+
+	/** Translation applied to UVs after packing, and after scaling */
+	UPROPERTY(EditAnywhere, Category = "UV Layout")
+	FVector2D Translation = FVector2D(0, 0);
+
+	/** Allow the Repack layout type to flip the orientation of UV islands to save space. Note that this may cause problems for downstream operations, and therefore is disabled by default. */
+	UPROPERTY(EditAnywhere, Category = "UV Layout", meta = (EditCondition = "LayoutType == EUVEditorUVLayoutType::Repack"))
+	bool bAllowFlips = false;
+
+	/** Enable UDIM aware layout and keep islands within their originating UDIM tiles when laying out.*/
+	UPROPERTY(EditAnywhere, Category = "UV Layout", meta = (DisplayName = "Preserve UDIMs", EditCondition="bUDIMCVAREnabled", EditConditionHides, HideEditConditionToggle = true))
+	bool bEnableUDIMLayout = false;
+
+	UPROPERTY(Transient)
+	bool bUDIMCVAREnabled = false;
+};
+
+
+namespace UE
+{
+namespace Geometry
+{
+	class FDynamicMesh3;
+	class FDynamicMeshUVPacker;
+
+enum class EUVEditorUVLayoutOpLayoutModes
+{
+	TransformOnly = 0,
+	RepackToUnitRect = 1,
+	StackInUnitRect = 2
+};
+
+/**
+ * This is an intentional duplication of the original FUVLayoutOp class from GeometryTools as we explore UDIMs.
+ * It may be temporary (with changes folded into the original) or may evolve into a UVEditor specific layout operation.
+ */
+class UVEDITORTOOLS_API FUVEditorUVLayoutOp : public FDynamicMeshOperator
+{
+public:
+	virtual ~FUVEditorUVLayoutOp() {}
+
+	// inputs
+	TSharedPtr<FDynamicMesh3, ESPMode::ThreadSafe> OriginalMesh;
+
+	EUVEditorUVLayoutOpLayoutModes UVLayoutMode = EUVEditorUVLayoutOpLayoutModes::RepackToUnitRect;
+
+	int UVLayerIndex = 0;
+	int TextureResolution = 128;
+	bool bAllowFlips = false;
+	bool bAlwaysSplitBowties = true;
+	float UVScaleFactor = 1.0;
+	float GutterSize = 1.0;
+	bool bMaintainOriginatingUDIM = false;
+	FVector2f UVTranslation = FVector2f::Zero();
+
+	void SetTransform(const FTransformSRT3d& Transform);
+
+	//
+	// FDynamicMeshOperator implementation
+	// 
+
+	virtual void CalculateResult(FProgressCancel* Progress) override;
+
+protected:
+	void ExecutePacker(FDynamicMeshUVPacker& Packer);
+};
+
+} // end namespace UE::Geometry
+} // end namespace UE
+
+/**
+ * Can be hooked up to a UMeshOpPreviewWithBackgroundCompute to perform UV Layout operations.
+ *
+ * Inherits from UObject so that it can hold a strong pointer to the settings UObject, which
+ * needs to be a UObject to be displayed in the details panel.
+ */
+UCLASS()
+class UVEDITORTOOLS_API UUVEditorUVLayoutOperatorFactory : public UObject, public UE::Geometry::IDynamicMeshOperatorFactory
+{
+	GENERATED_BODY()
+
+public:
+	// IDynamicMeshOperatorFactory API
+	virtual TUniquePtr<UE::Geometry::FDynamicMeshOperator> MakeNewOperator() override;
+
+	UPROPERTY()
+	TObjectPtr<UUVEditorUVLayoutProperties> Settings;
+
+	TSharedPtr<UE::Geometry::FDynamicMesh3, ESPMode::ThreadSafe> OriginalMesh;
+	TUniqueFunction<int32()> GetSelectedUVChannel = []() { return 0; };
+	FTransform TargetTransform;
+};
\ No newline at end of file
diff --git a/Engine/Plugins/Experimental/UVEditor/Source/UVEditorTools/Public/UVEditorLayoutTool.h b/Engine/Plugins/Experimental/UVEditor/Source/UVEditorTools/Public/UVEditorLayoutTool.h
--- a/Engine/Plugins/Experimental/UVEditor/Source/UVEditorTools/Public/UVEditorLayoutTool.h
+++ b/Engine/Plugins/Experimental/UVEditor/Source/UVEditorTools/Public/UVEditorLayoutTool.h
@@ -11,8 +11,8 @@
 #include "UVEditorLayoutTool.generated.h"
 
 class UUVEditorToolMeshInput;
-class UUVLayoutProperties;
-class UUVLayoutOperatorFactory;
+class UUVEditorUVLayoutProperties;
+class UUVEditorUVLayoutOperatorFactory;
 
 UCLASS()
 class UVEDITORTOOLS_API UUVEditorLayoutToolBuilder : public UInteractiveToolBuilder
@@ -62,10 +62,10 @@ protected:
 	TArray<TObjectPtr<UUVEditorToolMeshInput>> Targets;
 
 	UPROPERTY()
-	TObjectPtr<UUVLayoutProperties> Settings = nullptr;
+	TObjectPtr<UUVEditorUVLayoutProperties> Settings = nullptr;
 
 	UPROPERTY()
-	TArray<TObjectPtr<UUVLayoutOperatorFactory>> Factories;
+	TArray<TObjectPtr<UUVEditorUVLayoutOperatorFactory>> Factories;
 
 	//
 	// Analytics
diff --git a/Engine/Plugins/Experimental/UVEditor/Source/UVEditorTools/Public/UVEditorLogging.h b/Engine/Plugins/Experimental/UVEditor/Source/UVEditorTools/Public/UVEditorLogging.h
new file mode 100644
--- /dev/null
+++ b/Engine/Plugins/Experimental/UVEditor/Source/UVEditorTools/Public/UVEditorLogging.h
@@ -0,0 +1,8 @@
+// Copyright Epic Games, Inc. All Rights Reserved.
+
+#pragma once
+
+#include "CoreMinimal.h"
+
+UVEDITORTOOLS_API DECLARE_LOG_CATEGORY_EXTERN(LogUVEditor, Log, All);
+
diff --git a/Engine/Plugins/Experimental/UVEditor/Source/UVEditorTools/Public/Utilities/MeshUDIMClassifier.h b/Engine/Plugins/Experimental/UVEditor/Source/UVEditorTools/Public/Utilities/MeshUDIMClassifier.h
new file mode 100644
--- /dev/null
+++ b/Engine/Plugins/Experimental/UVEditor/Source/UVEditorTools/Public/Utilities/MeshUDIMClassifier.h
@@ -0,0 +1,38 @@
+// Copyright Epic Games, Inc. All Rights Reserved.
+
+#pragma once
+
+#include "DynamicMesh/DynamicMesh3.h"
+#include "DynamicMesh/DynamicMeshAttributeSet.h"
+
+
+namespace UE
+{
+	namespace Geometry
+	{
+
+		struct FUVOverlayView;
+		class FMeshConnectedComponents;
+
+		/**
+		* FDynamicMeshUDIMClassifier is a utility class for identifying active UDIMs from a FDynamicMesh's UV overlay
+		 */
+		class UVEDITORTOOLS_API FDynamicMeshUDIMClassifier
+		{
+		public:
+			explicit FDynamicMeshUDIMClassifier(FDynamicMeshUVOverlay* UVOverlay);
+
+			TArray<FVector2i> ActiveTiles() const;
+			TArray<int32> TidsForTile(FVector2i TileIndexIn) const;
+
+		protected:
+
+			void ClassifyUDIMs();
+
+			/** The UV Overlay to analyze for UDIMs */
+			FDynamicMeshUVOverlay* UVOverlay = nullptr;
+
+			TMap<FVector2i, TArray<int32> > UDIMs;
+		};
+	}
+}
\ No newline at end of file
diff --git a/Engine/Plugins/Runtime/GeometryProcessing/Source/DynamicMesh/Private/Parameterization/MeshUVPacking.cpp b/Engine/Plugins/Runtime/GeometryProcessing/Source/DynamicMesh/Private/Parameterization/MeshUVPacking.cpp
--- a/Engine/Plugins/Runtime/GeometryProcessing/Source/DynamicMesh/Private/Parameterization/MeshUVPacking.cpp
+++ b/Engine/Plugins/Runtime/GeometryProcessing/Source/DynamicMesh/Private/Parameterization/MeshUVPacking.cpp
@@ -10,48 +10,55 @@
 
 using namespace UE::Geometry;
 
-struct FUVOverlayView : public FUVPacker::IUVMeshView
+namespace UE::Geometry
 {
-	FDynamicMesh3* Mesh;
-	FDynamicMeshUVOverlay* UVOverlay;
-
-	FUVOverlayView(FDynamicMeshUVOverlay* UVOverlay) : UVOverlay(UVOverlay)
+	struct FUVOverlayView : public FUVPacker::IUVMeshView
 	{
-		Mesh = UVOverlay->GetParentMesh();
-	}
+		FDynamicMesh3* Mesh;
+		FDynamicMeshUVOverlay* UVOverlay;
 
-	virtual FIndex3i GetTriangle(int32 TID) const
-	{
-		return Mesh->GetTriangle(TID);
-	}
+		FUVOverlayView(FDynamicMeshUVOverlay* UVOverlay) : UVOverlay(UVOverlay)
+		{
+			Mesh = UVOverlay->GetParentMesh();
+		}
 
-	virtual FIndex3i GetUVTriangle(int32 TID) const
-	{
-		return UVOverlay->GetTriangle(TID);
-	}
+		virtual FIndex3i GetTriangle(int32 TID) const
+		{
+			return Mesh->GetTriangle(TID);
+		}
 
-	virtual FVector3d GetVertex(int32 VID) const
-	{
-		return Mesh->GetVertex(VID);
-	}
+		virtual FIndex3i GetUVTriangle(int32 TID) const
+		{
+			return UVOverlay->GetTriangle(TID);
+		}
 
-	virtual FVector2f GetUV(int32 EID) const
-	{
-		return UVOverlay->GetElement(EID);
-	}
+		virtual FVector3d GetVertex(int32 VID) const
+		{
+			return Mesh->GetVertex(VID);
+		}
 
-	virtual void SetUV(int32 EID, FVector2f UV)
-	{
-		return UVOverlay->SetElement(EID, UV);
-	}
-};
+		virtual FVector2f GetUV(int32 EID) const
+		{
+			return UVOverlay->GetElement(EID);
+		}
 
+		virtual void SetUV(int32 EID, FVector2f UV)
+		{
+			return UVOverlay->SetElement(EID, UV);
+		}
+	};
+}
 
 FDynamicMeshUVPacker::FDynamicMeshUVPacker(FDynamicMeshUVOverlay* UVOverlayIn)
 {
 	UVOverlay = UVOverlayIn;
 }
 
+FDynamicMeshUVPacker::FDynamicMeshUVPacker(FDynamicMeshUVOverlay* UVOverlayIn, TUniquePtr<TArray<int32>>&& TidsToRepackIn)
+{
+	UVOverlay = UVOverlayIn;
+	TidsToRepack = MoveTemp(TidsToRepackIn);
+}
 
 bool FDynamicMeshUVPacker::StandardPack()
 {
@@ -63,11 +70,7 @@ bool FDynamicMeshUVPacker::StandardPack()
 	Packer.TextureResolution = TextureResolution;
 
 	FUVOverlayView MeshView(UVOverlay);
-
-	FMeshConnectedComponents UVComponents(MeshView.Mesh);
-	UVComponents.FindConnectedTriangles([&MeshView](int32 Triangle0, int32 Triangle1) {
-		return MeshView.UVOverlay->AreTrianglesConnected(Triangle0, Triangle1);
-		});
+	FMeshConnectedComponents UVComponents = CollectUVIslandsToPack(MeshView);
 
 	return Packer.StandardPack(&MeshView, UVComponents.Num(), [&UVComponents](int Idx, TArray<int32>& Island)
 		{
@@ -75,8 +78,6 @@ bool FDynamicMeshUVPacker::StandardPack()
 		});
 }
 
-
-
 bool FDynamicMeshUVPacker::StackPack()
 {
 	TRACE_CPUPROFILER_EVENT_SCOPE(DynamicMeshUVPacker_StackPack);
@@ -87,14 +88,30 @@ bool FDynamicMeshUVPacker::StackPack()
 	Packer.TextureResolution = TextureResolution;
 
 	FUVOverlayView MeshView(UVOverlay);
-
-	FMeshConnectedComponents UVComponents(MeshView.Mesh);
-	UVComponents.FindConnectedTriangles([&MeshView](int32 Triangle0, int32 Triangle1) {
-		return MeshView.UVOverlay->AreTrianglesConnected(Triangle0, Triangle1);
-		});
+	FMeshConnectedComponents UVComponents = CollectUVIslandsToPack(MeshView);
 
 	return Packer.StackPack(&MeshView, UVComponents.Num(), [&UVComponents](int Idx, TArray<int32>& Island)
 		{
 			Island = UVComponents.GetComponent(Idx).Indices;
 		});
+}
+
+FMeshConnectedComponents FDynamicMeshUVPacker::CollectUVIslandsToPack(const FUVOverlayView& MeshView)
+{
+	auto UVIslandPredicate = [&MeshView](int32 Triangle0, int32 Triangle1)
+	{
+		return MeshView.UVOverlay->AreTrianglesConnected(Triangle0, Triangle1);
+	};
+
+	FMeshConnectedComponents UVComponents(MeshView.Mesh);
+	if (TidsToRepack)
+	{
+		UVComponents.FindConnectedTriangles(*TidsToRepack, UVIslandPredicate);
+	}
+	else
+	{
+		UVComponents.FindConnectedTriangles(UVIslandPredicate);
+	}
+
+	return UVComponents;
 }
\ No newline at end of file
diff --git a/Engine/Plugins/Runtime/GeometryProcessing/Source/DynamicMesh/Public/Parameterization/MeshUVPacking.h b/Engine/Plugins/Runtime/GeometryProcessing/Source/DynamicMesh/Public/Parameterization/MeshUVPacking.h
--- a/Engine/Plugins/Runtime/GeometryProcessing/Source/DynamicMesh/Public/Parameterization/MeshUVPacking.h
+++ b/Engine/Plugins/Runtime/GeometryProcessing/Source/DynamicMesh/Public/Parameterization/MeshUVPacking.h
@@ -5,11 +5,14 @@
 #include "DynamicMesh/DynamicMesh3.h"
 #include "DynamicMesh/DynamicMeshAttributeSet.h"
 
+
 namespace UE
 {
 namespace Geometry
 {
 
+struct FUVOverlayView;
+class FMeshConnectedComponents;
 
 /**
  * FDynamicMeshUVPacker implements various strategies for packing UV islands in a 
@@ -22,6 +25,9 @@ public:
 	/** The UV Overlay we will be repacking */
 	FDynamicMeshUVOverlay* UVOverlay = nullptr;
 
+	/** The explicit triangle ids to repack, repack all triangles if null */
+	TUniquePtr<TArray<int32>> TidsToRepack;
+
 	/** Resolution of the target texture. This is used to convert pixel gutter/border thickness to UV space */
 	int32 TextureResolution = 512;
 
@@ -32,6 +38,7 @@ public:
 	bool bAllowFlips = false;
 
 	explicit FDynamicMeshUVPacker(FDynamicMeshUVOverlay* UVOverlay);
+	explicit FDynamicMeshUVPacker(FDynamicMeshUVOverlay* UVOverlay, TUniquePtr<TArray<int32>>&& TidsToRepackIn);
 
 
 	/**
@@ -48,6 +55,10 @@ public:
 	 */
 	bool StackPack();
 
+protected:
+
+	FMeshConnectedComponents CollectUVIslandsToPack(const FUVOverlayView& MeshView);
+
 };
 
 
