UVEditor: Add support for UDIM tile aware processing in the Unwrap tool. Unwrap now processes, when the option is enabled, each UDIM tile separately, handling unwrapping and layout per tile.

diff --git a/Engine/Plugins/Experimental/UVEditor/Source/UVEditorTools/Private/Operators/UVEditorRecomputeUVsOp.cpp b/Engine/Plugins/Experimental/UVEditor/Source/UVEditorTools/Private/Operators/UVEditorRecomputeUVsOp.cpp
new file mode 100644
--- /dev/null
+++ b/Engine/Plugins/Experimental/UVEditor/Source/UVEditorTools/Private/Operators/UVEditorRecomputeUVsOp.cpp
@@ -0,0 +1,509 @@
+// Copyright Epic Games, Inc. All Rights Reserved.
+
+#include "Operators/UVEditorRecomputeUVsOp.h"
+
+#include "DynamicMesh/DynamicMeshAttributeSet.h"
+#include "DynamicMesh/MeshNormals.h"
+#include "Selections/MeshConnectedComponents.h"
+#include "Parameterization/MeshLocalParam.h"
+#include "Parameterization/DynamicMeshUVEditor.h"
+#include "Parameterization/PatchBasedMeshUVGenerator.h"
+#include "Properties/RecomputeUVsProperties.h"
+#include "Utilities/MeshUDIMClassifier.h"
+
+#include "DynamicSubmesh3.h"
+
+#include "Async/ParallelFor.h"
+
+using namespace UE::Geometry;
+
+
+#define LOCTEXT_NAMESPACE "UVEditorRecomputeUVsOp"
+namespace UE
+{
+	namespace Geometry
+	{
+		struct TileConnectedComponents
+		{
+			TileConnectedComponents(const FMeshConnectedComponents& ConnectedComponentsIn, const FVector2i& TileIn)
+				: ConnectedComponents(ConnectedComponentsIn),
+				  Tile(TileIn)
+			{
+				for (int32 ComponentIndex = 0; ComponentIndex < ConnectedComponents.Num(); ++ComponentIndex)
+				{
+					TileTids.Append(ConnectedComponents[ComponentIndex].Indices);
+				}
+			}
+			
+			TArray<int32> TileTids;
+			FMeshConnectedComponents ConnectedComponents;
+			FVector2i Tile;
+		};
+	}
+}
+
+void FUVEditorRecomputeUVsOp::NormalizeUVAreas(const FDynamicMesh3& Mesh, FDynamicMeshUVOverlay* Overlay, const TileConnectedComponents& TileComponents, float GlobalScale)
+{
+	int32 NumComponents = TileComponents.ConnectedComponents.Num();
+
+	ParallelFor(NumComponents, [&](int32 k)
+		{
+			const TArray<int>& Triangles = TileComponents.ConnectedComponents[k].Indices;
+			TSet<int> UVElements;
+			UVElements.Reserve(Triangles.Num() * 3);
+			double AreaUV = 0;
+			double Area3D = 0;
+			FVector3d Triangle3D[3];
+			FVector2d TriangleUV[3];
+			FAxisAlignedBox2d BoundsUV = FAxisAlignedBox2d::Empty();
+
+			for (int tid : Triangles)
+			{
+				FIndex3i TriElements = Overlay->GetTriangle(tid);
+				if (!TriElements.Contains(FDynamicMesh3::InvalidID))
+				{
+					for (int j = 0; j < 3; ++j)
+					{
+						TriangleUV[j] = FVector2d(Overlay->GetElement(TriElements[j]));
+						BoundsUV.Contain(TriangleUV[j]);
+						Triangle3D[j] = Mesh.GetVertex(Overlay->GetParentVertex(TriElements[j]));
+						UVElements.Add(TriElements[j]);
+					}
+					AreaUV += VectorUtil::Area(TriangleUV[0], TriangleUV[1], TriangleUV[2]);
+					Area3D += VectorUtil::Area(Triangle3D[0], Triangle3D[1], Triangle3D[2]);
+				}
+			}
+
+			double LinearScale = (AreaUV > 0.00001) ? (FMathd::Sqrt(Area3D) / FMathd::Sqrt(AreaUV)) : 1.0;
+			LinearScale = LinearScale * GlobalScale;
+			FVector2d ComponentOrigin = BoundsUV.Center();
+
+			for (int elemid : UVElements)
+			{
+				FVector2d UV = FVector2d(Overlay->GetElement(elemid));
+				UV = (UV - ComponentOrigin) * LinearScale;
+				Overlay->SetElement(elemid, FVector2f(UV));
+			}
+		});
+}
+
+
+void FUVEditorRecomputeUVsOp::CalculateResult(FProgressCancel* Progress)
+{
+	NewResultInfo = FGeometryResult(EGeometryResultType::InProgress);
+
+	bool bOK;
+	if (bMergingOptimization)
+	{
+		bOK = CalculateResult_RegionOptimization(Progress);
+	}
+	else
+	{
+		bOK = CalculateResult_Basic(Progress);
+	}
+
+	NewResultInfo.SetSuccess(bOK, Progress);
+	SetResultInfo(NewResultInfo);
+
+}
+
+
+void FUVEditorRecomputeUVsOp::CollectIslandComponentsPerTile(const FDynamicMeshUVOverlay& UVOverlay, TArray< TileConnectedComponents >& ComponentsPerTile, bool& bUseExistingUVTopology)
+{
+	bUseExistingUVTopology = false;
+	ComponentsPerTile.Empty();
+
+	TArray<FVector2i> Tiles;
+	TArray<TUniquePtr<TArray<int32>>> TileTids;
+
+	FDynamicMeshUDIMClassifier TileClassifier(&UVOverlay);
+
+	if (bUDIMsEnabled)
+	{
+		Tiles = TileClassifier.ActiveTiles();
+		for (const FVector2i& Tile : Tiles)
+		{
+			TileTids.Emplace(MakeUnique<TArray<int32>>(TileClassifier.TidsForTile(Tile)));
+		}
+	}
+	else
+	{
+		Tiles.Add({ 0,0 });
+		TileTids.Add(nullptr);
+	}
+	
+	for (int32 TileIndex = 0; TileIndex < Tiles.Num(); ++TileIndex)
+	{
+		FMeshConnectedComponents ConnectedComponents(ResultMesh.Get());
+		if (IslandMode == EUVEditorRecomputeUVsIslandMode::PolyGroups)
+		{
+			if (InputGroups != nullptr)
+			{
+				if (TileTids[TileIndex])
+				{
+					ConnectedComponents.FindConnectedTriangles(*TileTids[TileIndex], [this](int32 CurTri, int32 NbrTri) {
+						return InputGroups->GetTriangleGroup(CurTri) == InputGroups->GetTriangleGroup(NbrTri);
+						});
+				}
+				else
+				{
+					ConnectedComponents.FindConnectedTriangles([this](int32 CurTri, int32 NbrTri) {
+						return InputGroups->GetTriangleGroup(CurTri) == InputGroups->GetTriangleGroup(NbrTri);
+						});
+				}
+			}
+			else
+			{
+				if (TileTids[TileIndex])
+				{
+					ConnectedComponents.FindConnectedTriangles(*TileTids[TileIndex], [this](int32 CurTri, int32 NbrTri) {
+						return ResultMesh->GetTriangleGroup(CurTri) == ResultMesh->GetTriangleGroup(NbrTri);
+						});
+				}
+				else
+				{
+					ConnectedComponents.FindConnectedTriangles([this](int32 CurTri, int32 NbrTri) {
+						return ResultMesh->GetTriangleGroup(CurTri) == ResultMesh->GetTriangleGroup(NbrTri);
+						});
+				}
+			}
+		}
+		else
+		{
+			if (TileTids[TileIndex])
+			{
+				ConnectedComponents.FindConnectedTriangles(*TileTids[TileIndex], [&](int32 Triangle0, int32 Triangle1) {
+					return UVOverlay.AreTrianglesConnected(Triangle0, Triangle1);
+					});
+			}
+			else
+			{
+				ConnectedComponents.FindConnectedTriangles([&](int32 Triangle0, int32 Triangle1) {
+					return UVOverlay.AreTrianglesConnected(Triangle0, Triangle1);
+					});
+			}
+
+			bUseExistingUVTopology = true;
+		}
+		ComponentsPerTile.Emplace(ConnectedComponents, Tiles[TileIndex]);
+	}
+}
+
+bool FUVEditorRecomputeUVsOp::CalculateResult_Basic(FProgressCancel* Progress)
+{
+	if (!InputMesh.IsValid())
+	{
+		return false;
+	}
+
+	ResultMesh = MakeUnique<FDynamicMesh3>(*InputMesh);
+
+	FDynamicMesh3* BaseMesh = ResultMesh.Get();
+	FDynamicMeshUVEditor UVEditor(BaseMesh, UVLayer, true);
+	FDynamicMeshUVOverlay* UseOverlay = UVEditor.GetOverlay();
+	if (ensure(UseOverlay != nullptr) == false)
+	{
+		return false;
+	}
+
+	if (Progress && Progress->Cancelled())
+	{
+		return false;
+	}
+
+	// find group-connected-components
+	bool bUseExistingUVTopology = false;
+	TArray< TileConnectedComponents > ComponentsPerTile;
+	CollectIslandComponentsPerTile(*UseOverlay, ComponentsPerTile, bUseExistingUVTopology);
+
+	if (Progress && Progress->Cancelled())
+	{
+		return false;
+	}
+
+	
+	for (TileConnectedComponents& TileComponents : ComponentsPerTile)
+	{
+		int32 NumComponents = TileComponents.ConnectedComponents.Num();
+		TArray<bool> bComponentSolved;
+		bComponentSolved.Init(false, NumComponents);
+		int32 SuccessCount = 0;
+
+		// TODO: the solves here could be done in parallel if we pre-allocated the island element IDs
+		for (int32 k = 0; k < NumComponents; ++k)
+		{
+			const TArray<int32>& ComponentTris = TileComponents.ConnectedComponents[k].Indices;
+
+			bComponentSolved[k] = false;
+			switch (UnwrapType)
+			{
+			case EUVEditorRecomputeUVsUnwrapType::ExpMap:
+			{
+				FDynamicMeshUVEditor::FExpMapOptions Options;
+				Options.NormalSmoothingRounds = this->NormalSmoothingRounds;
+				Options.NormalSmoothingAlpha = this->NormalSmoothingAlpha;
+				bComponentSolved[k] = UVEditor.SetTriangleUVsFromExpMap(ComponentTris, Options);
+			}
+			break;
+
+			case EUVEditorRecomputeUVsUnwrapType::ConformalFreeBoundary:
+				bComponentSolved[k] = UVEditor.SetTriangleUVsFromFreeBoundaryConformal(ComponentTris, bUseExistingUVTopology);
+				if (bComponentSolved[k])
+				{
+					UVEditor.ScaleUVAreaTo3DArea(ComponentTris, true);
+				}
+				break;
+			}
+
+			if (bComponentSolved[k])
+			{
+				SuccessCount++;
+			}
+
+			if (Progress && Progress->Cancelled())
+			{
+				return false;
+			}
+		}
+
+
+		if (bAutoRotate)
+		{
+			ParallelFor(NumComponents, [&](int32 k)
+				{
+					if (Progress && Progress->Cancelled())
+					{
+						return;
+					}
+					if (bComponentSolved[k])
+					{
+						const TArray<int32>& ComponentTris = TileComponents.ConnectedComponents[k].Indices;
+						UVEditor.AutoOrientUVArea(ComponentTris);
+					};
+				});
+		}
+		if (Progress && Progress->Cancelled())
+		{
+			return false;
+		}
+
+		if (bNormalizeAreas)
+		{
+			// todo should be a DynamicUVEditor function?
+			NormalizeUVAreas(*ResultMesh, UseOverlay, TileComponents, AreaScaling);
+		}
+
+		if (Progress && Progress->Cancelled())
+		{
+			return false;
+		}
+
+		if (bPackUVs)
+		{
+			bool bPackingSuccess = UVEditor.UDIMPack(PackingTextureResolution, PackingGutterWidth, TileComponents.Tile, &TileComponents.TileTids);
+			if (!bPackingSuccess)
+			{
+				NewResultInfo.AddError(FGeometryError(0, LOCTEXT("IslandPackingFailed", "Failed to pack UV islands")));
+				return false;
+			}
+		}
+
+		if (Progress && Progress->Cancelled())
+		{
+			return false;
+		}
+
+	}
+
+	return true;
+}
+
+
+
+
+
+
+bool FUVEditorRecomputeUVsOp::CalculateResult_RegionOptimization(FProgressCancel* Progress)
+{
+	if (!InputMesh.IsValid())
+	{
+		return false;
+	}
+
+	ResultMesh = MakeUnique<FDynamicMesh3>(*InputMesh);
+
+	FDynamicMesh3* BaseMesh = ResultMesh.Get();
+	FDynamicMeshUVEditor UVEditor(BaseMesh, UVLayer, true);
+	FDynamicMeshUVOverlay* UseOverlay = UVEditor.GetOverlay();
+	if (ensure(UseOverlay != nullptr) == false)
+	{
+		return false;
+	}
+
+	if (Progress && Progress->Cancelled())
+	{
+		return false;
+	}
+
+	// find group-connected-components
+	bool bUseExistingUVTopology = false;
+	TArray< TileConnectedComponents > ComponentsPerTile;
+	CollectIslandComponentsPerTile(*UseOverlay, ComponentsPerTile, bUseExistingUVTopology);
+
+	if (Progress && Progress->Cancelled())
+	{
+		return false;
+	}
+
+	for (TileConnectedComponents& TileComponents : ComponentsPerTile)
+	{
+		FPatchBasedMeshUVGenerator UVGenerator;
+		UVGenerator.MergingThreshold = this->MergingThreshold;
+		UVGenerator.CompactnessThreshold = this->CompactnessThreshold;
+		UVGenerator.MaxNormalDeviationDeg = this->MaxNormalDeviationDeg;
+		UVGenerator.NormalSmoothingRounds = this->NormalSmoothingRounds;
+		UVGenerator.NormalSmoothingAlpha = this->NormalSmoothingAlpha;
+
+		TArray<TArray<int32>> UVIslandTriangleSets;
+		bool bIslandsOK = UVGenerator.ComputeIslandsByRegionMerging(*BaseMesh, *UseOverlay, TileComponents.ConnectedComponents, UVIslandTriangleSets, Progress);
+
+		if (bIslandsOK == false)
+		{
+			NewResultInfo.AddError(FGeometryError(0, LOCTEXT("IslandMergingFailed", "Failed to merge UV islands")));
+			return false;
+		}
+
+		if (Progress && Progress->Cancelled())
+		{
+			return false;
+		}
+
+		TArray<bool> bIslandUVsValid;
+		int32 NumSolvesFailed = UVGenerator.ComputeUVsFromTriangleSets(*BaseMesh, *UseOverlay, UVIslandTriangleSets, bIslandUVsValid, Progress);
+		if (NumSolvesFailed > 0)
+		{
+			NewResultInfo.AddWarning(FGeometryWarning(0, LOCTEXT("PartialSolvesFailed", "Failed to compute UVs for some UV islands")));
+		}
+
+		if (Progress && Progress->Cancelled())
+		{
+			return false;
+		}
+
+		if (bAutoRotate)
+		{
+			ParallelFor(UVIslandTriangleSets.Num(), [&](int32 k)
+				{
+					if (Progress && Progress->Cancelled())
+					{
+						return;
+					}
+					if (bIslandUVsValid[k])
+					{
+						UVEditor.AutoOrientUVArea(UVIslandTriangleSets[k]);
+					}
+				});
+		}
+		if (Progress && Progress->Cancelled())
+		{
+			return false;
+		}
+		
+		if (bNormalizeAreas)
+		{
+			// todo should be a DynamicUVEditor function?
+			NormalizeUVAreas(*ResultMesh, UseOverlay, TileComponents, AreaScaling);
+		}
+		
+
+		if (Progress && Progress->Cancelled())
+		{
+			return false;
+		}
+
+		if (bPackUVs)
+		{
+			bool bPackingSuccess = UVEditor.UDIMPack(PackingTextureResolution, PackingGutterWidth, TileComponents.Tile, &TileComponents.TileTids);
+			if (!bPackingSuccess)
+			{
+				NewResultInfo.AddError(FGeometryError(0, LOCTEXT("IslandPackingFailed", "Failed to pack UV islands")));
+				return false;
+			}
+		}
+
+		if (Progress && Progress->Cancelled())
+		{
+			return false;
+		}
+	}
+
+	return true;
+}
+
+
+/**
+ * Factory
+ */
+
+
+TUniquePtr<FDynamicMeshOperator> UUVEditorRecomputeUVsOpFactory::MakeNewOperator()
+{
+	TUniquePtr<FUVEditorRecomputeUVsOp> RecomputeUVsOp = MakeUnique<FUVEditorRecomputeUVsOp>();
+
+	RecomputeUVsOp->InputMesh = OriginalMesh;
+	RecomputeUVsOp->InputGroups = InputGroups;
+	RecomputeUVsOp->UVLayer = GetSelectedUVChannel();
+
+	RecomputeUVsOp->IslandMode = static_cast<EUVEditorRecomputeUVsIslandMode>(Settings->IslandGeneration);
+
+	switch (Settings->UnwrapType)
+	{
+	case EUVEditorRecomputeUVsPropertiesUnwrapType::ExpMap:
+		RecomputeUVsOp->UnwrapType = EUVEditorRecomputeUVsUnwrapType::ExpMap;
+		RecomputeUVsOp->bMergingOptimization = false;
+		break;
+	case EUVEditorRecomputeUVsPropertiesUnwrapType::Conformal:
+		RecomputeUVsOp->UnwrapType = EUVEditorRecomputeUVsUnwrapType::ConformalFreeBoundary;
+		RecomputeUVsOp->bMergingOptimization = false;
+		break;
+	case EUVEditorRecomputeUVsPropertiesUnwrapType::IslandMerging:
+		RecomputeUVsOp->UnwrapType = EUVEditorRecomputeUVsUnwrapType::ExpMap;
+		RecomputeUVsOp->bMergingOptimization = true;
+	}
+
+	RecomputeUVsOp->bAutoRotate = (Settings->AutoRotation == EUVEditorRecomputeUVsToolOrientationMode::MinBounds);
+
+	RecomputeUVsOp->NormalSmoothingRounds = Settings->SmoothingSteps;
+	RecomputeUVsOp->NormalSmoothingAlpha = Settings->SmoothingAlpha;
+
+	RecomputeUVsOp->MergingThreshold = Settings->MergingDistortionThreshold;
+	RecomputeUVsOp->MaxNormalDeviationDeg = Settings->MergingAngleThreshold;
+
+	RecomputeUVsOp->bPackUVs = false;
+	RecomputeUVsOp->bNormalizeAreas = false;
+	RecomputeUVsOp->bUDIMsEnabled = Settings->bEnableUDIMLayout;
+
+	switch (Settings->LayoutType)
+	{
+	case EUVEditorRecomputeUVsPropertiesLayoutType::None:
+		break;
+	case EUVEditorRecomputeUVsPropertiesLayoutType::Repack:
+		RecomputeUVsOp->bPackUVs = true;
+		RecomputeUVsOp->PackingTextureResolution = Settings->TextureResolution;
+		break;
+	case EUVEditorRecomputeUVsPropertiesLayoutType::NormalizeToBounds:
+		RecomputeUVsOp->bNormalizeAreas = true;
+		RecomputeUVsOp->AreaScaling = Settings->NormalizeScale / OriginalMesh->GetBounds().MaxDim();
+		break;
+	case EUVEditorRecomputeUVsPropertiesLayoutType::NormalizeToWorld:
+		RecomputeUVsOp->bNormalizeAreas = true;
+		RecomputeUVsOp->AreaScaling = Settings->NormalizeScale;
+	}
+
+	RecomputeUVsOp->SetTransform(TargetTransform);
+
+	return RecomputeUVsOp;
+}
+
+
+#undef LOCTEXT_NAMESPACE
diff --git a/Engine/Plugins/Experimental/UVEditor/Source/UVEditorTools/Private/UVEditorRecomputeUVsTool.cpp b/Engine/Plugins/Experimental/UVEditor/Source/UVEditorTools/Private/UVEditorRecomputeUVsTool.cpp
--- a/Engine/Plugins/Experimental/UVEditor/Source/UVEditorTools/Private/UVEditorRecomputeUVsTool.cpp
+++ b/Engine/Plugins/Experimental/UVEditor/Source/UVEditorTools/Private/UVEditorRecomputeUVsTool.cpp
@@ -5,13 +5,14 @@
 #include "ToolBuilderUtil.h"
 #include "DynamicMesh/DynamicMesh3.h"
 #include "Polygroups/PolygroupUtil.h"
-#include "ParameterizationOps/RecomputeUVsOp.h"
 #include "ToolTargets/UVEditorToolMeshInput.h"
 #include "ContextObjects/UVToolContextObjects.h"
 #include "ContextObjectStore.h"
 #include "MeshOpPreviewHelpers.h"
 #include "UVEditorToolAnalyticsUtils.h"
 #include "EngineAnalytics.h"
+#include "Operators/UVEditorRecomputeUVsOp.h"
+#include "UVEditorUXSettings.h"
 
 using namespace UE::Geometry;
 
@@ -49,8 +50,9 @@ void UUVEditorRecomputeUVsTool::Setup()
 
 	// initialize our properties
 
-	Settings = NewObject<URecomputeUVsToolProperties>(this);
+	Settings = NewObject<UUVEditorRecomputeUVsToolProperties>(this);
 	Settings->RestoreProperties(this);
+	Settings->bUDIMCVAREnabled = (FUVEditorUXSettings::CVarEnablePrototypeUDIMSupport.GetValueOnGameThread() > 0);
 	AddToolPropertySource(Settings);
 	Factories.SetNum(Targets.Num());
 
@@ -66,14 +68,14 @@ void UUVEditorRecomputeUVsTool::Setup()
 	else
 	{
 		Settings->bEnablePolygroupSupport = false;
-		Settings->IslandGeneration = ERecomputeUVsPropertiesIslandMode::ExistingUVs;
+		Settings->IslandGeneration = EUVEditorRecomputeUVsPropertiesIslandMode::ExistingUVs;
 	}
 
 	
 	for (int32 TargetIndex = 0; TargetIndex < Targets.Num(); ++TargetIndex)
 	{
 		TObjectPtr<UUVEditorToolMeshInput> Target = Targets[TargetIndex];
-		Factories[TargetIndex] = NewObject<URecomputeUVsOpFactory>();
+		Factories[TargetIndex] = NewObject<UUVEditorRecomputeUVsOpFactory>();
 		Factories[TargetIndex]->TargetTransform = (FTransform3d)Target->AppliedPreview->PreviewMesh->GetTransform();
 		Factories[TargetIndex]->Settings = Settings;
 		Factories[TargetIndex]->OriginalMesh = Target->AppliedCanonical;
@@ -264,23 +266,23 @@ void UUVEditorRecomputeUVsTool::RecordAnalytics()
 	Attributes.Add(AnalyticsEventAttributeEnum(TEXT("Settings.AutoRotation"), Settings->AutoRotation));
 	
 	Attributes.Add(AnalyticsEventAttributeEnum(TEXT("Settings.UnwrapType"), Settings->UnwrapType));
-	if (Settings->UnwrapType == ERecomputeUVsPropertiesUnwrapType::IslandMerging || Settings->UnwrapType == ERecomputeUVsPropertiesUnwrapType::ExpMap)
+	if (Settings->UnwrapType == EUVEditorRecomputeUVsPropertiesUnwrapType::IslandMerging || Settings->UnwrapType == EUVEditorRecomputeUVsPropertiesUnwrapType::ExpMap)
 	{
 		Attributes.Add(FAnalyticsEventAttribute(TEXT("Settings.SmoothingSteps"), Settings->SmoothingSteps));
 		Attributes.Add(FAnalyticsEventAttribute(TEXT("Settings.SmoothingAlpha"), Settings->SmoothingAlpha));
 	}
-	if (Settings->UnwrapType == ERecomputeUVsPropertiesUnwrapType::IslandMerging)
+	if (Settings->UnwrapType == EUVEditorRecomputeUVsPropertiesUnwrapType::IslandMerging)
 	{
 		Attributes.Add(FAnalyticsEventAttribute(TEXT("Settings.MergingDistortionThreshold"), Settings->MergingDistortionThreshold));
 		Attributes.Add(FAnalyticsEventAttribute(TEXT("Settings.MergingAngleThreshold"), Settings->MergingAngleThreshold));
 	}
 	
 	Attributes.Add(AnalyticsEventAttributeEnum(TEXT("Settings.LayoutType"), Settings->LayoutType));
-	if (Settings->LayoutType == ERecomputeUVsPropertiesLayoutType::Repack)
+	if (Settings->LayoutType == EUVEditorRecomputeUVsPropertiesLayoutType::Repack)
 	{
 		Attributes.Add(FAnalyticsEventAttribute(TEXT("Settings.TextureResolution"), Settings->TextureResolution));
 	}
-	if (Settings->LayoutType == ERecomputeUVsPropertiesLayoutType::NormalizeToBounds || Settings->LayoutType == ERecomputeUVsPropertiesLayoutType::NormalizeToWorld)
+	if (Settings->LayoutType == EUVEditorRecomputeUVsPropertiesLayoutType::NormalizeToBounds || Settings->LayoutType == EUVEditorRecomputeUVsPropertiesLayoutType::NormalizeToWorld)
 	{
 		Attributes.Add(FAnalyticsEventAttribute(TEXT("Settings.NormalizeScale"), Settings->NormalizeScale));
 	}
diff --git a/Engine/Plugins/Experimental/UVEditor/Source/UVEditorTools/Private/Utilities/MeshUDIMClassifier.cpp b/Engine/Plugins/Experimental/UVEditor/Source/UVEditorTools/Private/Utilities/MeshUDIMClassifier.cpp
--- a/Engine/Plugins/Experimental/UVEditor/Source/UVEditorTools/Private/Utilities/MeshUDIMClassifier.cpp
+++ b/Engine/Plugins/Experimental/UVEditor/Source/UVEditorTools/Private/Utilities/MeshUDIMClassifier.cpp
@@ -7,7 +7,7 @@
 
 using namespace UE::Geometry;
 
-FDynamicMeshUDIMClassifier::FDynamicMeshUDIMClassifier(FDynamicMeshUVOverlay* UVOverlayIn)
+FDynamicMeshUDIMClassifier::FDynamicMeshUDIMClassifier(const FDynamicMeshUVOverlay* UVOverlayIn)
 {
 	UVOverlay = UVOverlayIn;
 	ClassifyUDIMs();
@@ -35,7 +35,7 @@ TArray<int32> FDynamicMeshUDIMClassifier::TidsForTile(FVector2i TileIndexIn) con
 void FDynamicMeshUDIMClassifier::ClassifyUDIMs()
 {
 	int32 CurrentUDIMIndex = 0;
-	FDynamicMesh3* Mesh = UVOverlay->GetParentMesh();
+	const FDynamicMesh3* Mesh = UVOverlay->GetParentMesh();
 
 	auto UVIslandPredicate = [this](int32 Triangle0, int32 Triangle1)
 	{
diff --git a/Engine/Plugins/Experimental/UVEditor/Source/UVEditorTools/Public/Operators/UVEditorRecomputeUVsOp.h b/Engine/Plugins/Experimental/UVEditor/Source/UVEditorTools/Public/Operators/UVEditorRecomputeUVsOp.h
new file mode 100644
--- /dev/null
+++ b/Engine/Plugins/Experimental/UVEditor/Source/UVEditorTools/Public/Operators/UVEditorRecomputeUVsOp.h
@@ -0,0 +1,262 @@
+// Copyright Epic Games, Inc. All Rights Reserved.
+
+#pragma once
+
+#include "CoreMinimal.h"
+#include "Util/ProgressCancel.h"
+#include "ModelingOperators.h"
+#include "DynamicMesh/DynamicMeshAttributeSet.h"
+#include "Polygroups/PolygroupSet.h"
+#include "InteractiveTool.h"
+
+#include "UVEditorRecomputeUVsOp.generated.h"
+
+class URecomputeUVsToolProperties;
+
+/*
+ * TODO: This code is mostly copy and pasted from the original RecomputeUVsOp, for the
+ * purposes of exploring UDIM workflows in the UV Editor. Once a decision has been made
+ * regarding the path we want UDIMs to take, this can be merged back into the original
+ * Modeling Mode operator. 
+ */
+
+UENUM()
+enum class EUVEditorRecomputeUVsPropertiesUnwrapType
+{
+	// Values must match UE::Geometry::ERecomputeUVsUnwrapType
+
+	/** ExpMap UV flattening is very fast but has limited abilities to reduce stretching and distortion */
+	ExpMap = 0 UMETA(DisplayName = "ExpMap"),
+	/** Conformal UV flattening is increasingly expensive on large islands but reduces distortion */
+	Conformal = 1,
+	/** UV islands will be merged into larger islands if it does not increase stretching and distortion beyond defined limits */
+	IslandMerging = 2
+};
+
+
+UENUM()
+enum class EUVEditorRecomputeUVsPropertiesIslandMode
+{
+	// Values must match UE::Geometry::ERecomputeUVsIslandMode
+
+	/** Use Polygroups defined by Active PolyGroup property to define initial UV islands. */
+	PolyGroups = 0 UMETA(DisplayName = "PolyGroups"),
+	/** Use existing UV Layer to define UV islands, i.e. re-solve UV flattening based on existing UVs */
+	ExistingUVs = 1
+};
+
+
+UENUM()
+enum class EUVEditorRecomputeUVsToolOrientationMode
+{
+	/** Do not rotate UV islands */
+	None,
+	/** Automatically rotate each UV island to reduce its bounding box size */
+	MinBounds
+};
+
+
+UENUM()
+enum class EUVEditorRecomputeUVsPropertiesLayoutType
+{
+	/** Do not apply additional layout options to unwrapped UVs */
+	None,
+	/** Uniformly scale and translate UV islands collectively to pack them into the unit square, i.e. fit between 0 and 1 with no overlap */
+	Repack,
+	/** Uniformly scale UV islands such that they have constant relative area, relative to object bounds */
+	NormalizeToBounds,
+	/** Uniformly scale UV islands such that they have constant relative area, relative to world space */
+	NormalizeToWorld
+};
+
+
+UCLASS()
+class UVEDITORTOOLS_API  UUVEditorRecomputeUVsToolProperties : public UInteractiveToolPropertySet
+{
+	GENERATED_BODY()
+public:
+
+	UPROPERTY(meta = (TransientToolProperty))
+	bool bEnablePolygroupSupport = true;
+
+	/** Generation method for initial UV islands.*/
+	UPROPERTY(EditAnywhere, Category = "UV Unwrap", meta = (EditCondition = "bEnablePolygroupSupport", EditConditionHides, HideEditConditionToggle = true))
+	EUVEditorRecomputeUVsPropertiesIslandMode IslandGeneration = EUVEditorRecomputeUVsPropertiesIslandMode::ExistingUVs;
+
+	/** Type of UV flattening algorithm to use */
+	UPROPERTY(EditAnywhere, Category = "UV Unwrap")
+	EUVEditorRecomputeUVsPropertiesUnwrapType UnwrapType = EUVEditorRecomputeUVsPropertiesUnwrapType::Conformal;
+
+	/** Type of automatic rotation applied to each UV island */
+	UPROPERTY(EditAnywhere, Category = "UV Unwrap")
+	EUVEditorRecomputeUVsToolOrientationMode AutoRotation = EUVEditorRecomputeUVsToolOrientationMode::MinBounds;
+
+	/** Number of smoothing steps to apply; this slightly increases distortion but produces more stable results. This is only enabled if the Unwrap Type is set to ExpMap or Island Merging. */
+	UPROPERTY(EditAnywhere, Category = "UV Unwrap", meta = (UIMin = "0", UIMax = "25", ClampMin = "0", ClampMax = "1000",
+		EditCondition = "UnwrapType == EUVEditorRecomputeUVsPropertiesUnwrapType::ExpMap || UnwrapType == EUVEditorRecomputeUVsPropertiesUnwrapType::IslandMerging"))
+	int SmoothingSteps = 5;
+
+	/** Smoothing parameter; larger values result in faster smoothing in each step. This is only enabled if the Unwrap Type is set to ExpMap or Island Merging. */
+	UPROPERTY(EditAnywhere, Category = "UV Unwrap", meta = (UIMin = "0", UIMax = "1.0", ClampMin = "0", ClampMax = "1.0",
+		EditCondition = "UnwrapType == EUVEditorRecomputeUVsPropertiesUnwrapType::ExpMap || UnwrapType == EUVEditorRecomputeUVsPropertiesUnwrapType::IslandMerging"))
+	float SmoothingAlpha = 0.25f;
+
+	/** Threshold for stretching and distortion below which island merging is allowed; larger values increase the allowable UV distortion. This is only enabled if the Unwrap Type is set to Island Merging. */
+	UPROPERTY(EditAnywhere, Category = "UV Unwrap", meta = (DisplayName = "Distortion Threshold", UIMin = "1.0", UIMax = "5.0", ClampMin = "1.0",
+		EditCondition = "UnwrapType == EUVEditorRecomputeUVsPropertiesUnwrapType::IslandMerging"))
+	float MergingDistortionThreshold = 1.5f;
+
+	/** Threshold for the average face normal deviation below  which island merging is allowed. This is only enabled if the Unwrap Type is set to Island Merging. */
+	UPROPERTY(EditAnywhere, Category = "UV Unwrap", meta = (DisplayName = "Angle Threshold", UIMin = "0.0", UIMax = "90.0", ClampMin = "0.0", ClampMax = "180.0",
+		EditCondition = "UnwrapType == EUVEditorRecomputeUVsPropertiesUnwrapType::IslandMerging"))
+	float MergingAngleThreshold = 45.0f;
+
+	/** Uniformly scale and translate UV islands collectively to pack them into the unit square, i.e. fit between 0 and 1 with no overlap. */
+	UPROPERTY(EditAnywhere, Category = "UV Layout")
+	EUVEditorRecomputeUVsPropertiesLayoutType LayoutType = EUVEditorRecomputeUVsPropertiesLayoutType::Repack;
+
+	/** Expected resolution of the output textures; this controls spacing left between UV islands to avoid interpolation artifacts. This is only enabled when the Layout Type is set to Repack. */
+	UPROPERTY(EditAnywhere, Category = "UV Layout", meta = (UIMin = "64", UIMax = "2048", ClampMin = "2", ClampMax = "4096",
+		EditCondition = "LayoutType == EUVEditorRecomputeUVsPropertiesLayoutType::Repack"))
+	int TextureResolution = 1024;
+
+	/** Scaling factor used for UV island normalization/scaling. This is only enabled when the Layout Type is set to Normalize to Bounds or Normalize to World. */
+	UPROPERTY(EditAnywhere, Category = "UV Layout", meta = (UIMin = "0.001", UIMax = "10", ClampMin = "0.00001", ClampMax = "1000000.0",
+		EditCondition = "LayoutType == EUVEditorRecomputeUVsPropertiesLayoutType::NormalizeToBounds || LayoutType == EUVEditorRecomputeUVsPropertiesLayoutType::NormalizeToWorld"))
+	float NormalizeScale = 1.0f;
+
+	/** Enable UDIM aware layout and keep islands within their originating UDIM tiles when laying out.*/
+	UPROPERTY(EditAnywhere, Category = "UV Layout", meta = (DisplayName = "Preserve UDIMs", EditCondition = "bUDIMCVAREnabled", EditConditionHides, HideEditConditionToggle = true))
+	bool bEnableUDIMLayout = false;
+
+	UPROPERTY(Transient)
+	bool bUDIMCVAREnabled = false;
+};
+
+
+namespace UE
+{
+namespace Geometry
+{
+
+	struct TileConnectedComponents;
+
+enum class EUVEditorRecomputeUVsUnwrapType
+{
+	ExpMap = 0,
+	ConformalFreeBoundary = 1
+};
+
+
+
+enum class EUVEditorRecomputeUVsIslandMode
+{
+	PolyGroups = 0,
+	UVIslands = 1
+};
+
+
+class UVEDITORTOOLS_API FUVEditorRecomputeUVsOp : public FDynamicMeshOperator
+{
+public:
+	virtual ~FUVEditorRecomputeUVsOp() {}
+
+	//
+	// Inputs
+	// 
+
+	// source mesh
+	TSharedPtr<FDynamicMesh3, ESPMode::ThreadSafe> InputMesh;
+	
+	// source groups (optional)
+	TSharedPtr<UE::Geometry::FPolygroupSet, ESPMode::ThreadSafe> InputGroups;
+
+	// orientation control
+	bool bAutoRotate = true;
+
+	// area scaling
+	bool bNormalizeAreas = true;
+	float AreaScaling = 1.0;
+
+	// UV layer
+	int32 UVLayer = 0;
+
+	// Atlas Packing parameters
+	bool bPackUVs = true;
+	int32 PackingTextureResolution = 512;
+	float PackingGutterWidth = 1.0f;
+
+	EUVEditorRecomputeUVsUnwrapType UnwrapType = EUVEditorRecomputeUVsUnwrapType::ExpMap;
+	EUVEditorRecomputeUVsIslandMode IslandMode = EUVEditorRecomputeUVsIslandMode::PolyGroups;
+
+	// 
+	// ExpMap Options
+	//
+	int32 NormalSmoothingRounds = 0;
+	double NormalSmoothingAlpha = 0.25;
+
+	//
+	// Patch Merging options
+	//
+	bool bMergingOptimization = false;
+	double MergingThreshold = 1.5;
+	double CompactnessThreshold = 9999999.0;		// effectively disabled as it usually is not a good idea
+	double MaxNormalDeviationDeg = 45.0;
+
+	// UDIM options	
+	bool bUDIMsEnabled = false;
+
+
+	// set ability on protected transform.
+	void SetTransform(const FTransformSRT3d& XForm)
+	{
+		ResultTransform = XForm;
+	}
+
+	//
+	// FDynamicMeshOperator implementation
+	// 
+
+	virtual void CalculateResult(FProgressCancel* Progress) override;
+
+
+protected:
+
+	FGeometryResult NewResultInfo;
+
+	void NormalizeUVAreas(const FDynamicMesh3& Mesh, FDynamicMeshUVOverlay* Overlay, const TileConnectedComponents& TileComponents, float GlobalScale = 1.0f);
+
+	void CollectIslandComponentsPerTile(const FDynamicMeshUVOverlay& UVOverlay, TArray< TileConnectedComponents >& ComponentsPerTile, bool& UseExistingUVTopology);
+
+	virtual bool CalculateResult_Basic(FProgressCancel* Progress);
+	virtual bool CalculateResult_RegionOptimization(FProgressCancel* Progress);
+};
+
+} // end namespace UE::Geometry
+} // end namespace UE
+
+
+/**
+ * Can be hooked up to a UMeshOpPreviewWithBackgroundCompute to perform UV solving operations.
+ * 
+ * Inherits from UObject so that it can hold a strong pointer to the settings UObject, which
+ * needs to be a UObject to be displayed in the details panel.
+ */
+UCLASS()
+class UVEDITORTOOLS_API UUVEditorRecomputeUVsOpFactory : public UObject, public UE::Geometry::IDynamicMeshOperatorFactory
+{
+	GENERATED_BODY()
+
+public:
+
+	// IDynamicMeshOperatorFactory API
+	virtual TUniquePtr<UE::Geometry::FDynamicMeshOperator> MakeNewOperator() override;
+
+	UPROPERTY()
+	TObjectPtr<UUVEditorRecomputeUVsToolProperties> Settings = nullptr;
+
+	TSharedPtr<UE::Geometry::FPolygroupSet, ESPMode::ThreadSafe> InputGroups;
+	TSharedPtr<UE::Geometry::FDynamicMesh3, ESPMode::ThreadSafe> OriginalMesh;
+	TUniqueFunction<int32()> GetSelectedUVChannel = []() { return 0; };
+	UE::Geometry::FTransformSRT3d TargetTransform;
+};
\ No newline at end of file
diff --git a/Engine/Plugins/Experimental/UVEditor/Source/UVEditorTools/Public/UVEditorRecomputeUVsTool.h b/Engine/Plugins/Experimental/UVEditor/Source/UVEditorTools/Public/UVEditorRecomputeUVsTool.h
--- a/Engine/Plugins/Experimental/UVEditor/Source/UVEditorTools/Public/UVEditorRecomputeUVsTool.h
+++ b/Engine/Plugins/Experimental/UVEditor/Source/UVEditorTools/Public/UVEditorRecomputeUVsTool.h
@@ -6,7 +6,6 @@
 #include "BaseTools/SingleSelectionMeshEditingTool.h"
 #include "DynamicMesh/DynamicMesh3.h"
 #include "PropertySets/PolygroupLayersProperties.h"
-#include "Properties/RecomputeUVsProperties.h"
 #include "Polygroups/PolygroupSet.h"
 #include "Drawing/UVLayoutPreview.h"
 #include "UVEditorToolAnalyticsUtils.h"
@@ -19,7 +18,8 @@ class UDynamicMeshComponent;
 class UMaterialInterface;
 class UMaterialInstanceDynamic;
 class UUVEditorToolMeshInput;
-class URecomputeUVsOpFactory;
+class UUVEditorRecomputeUVsToolProperties;
+class UUVEditorRecomputeUVsOpFactory;
 
 /**
  *
@@ -75,13 +75,13 @@ protected:
 	TArray<TObjectPtr<UUVEditorToolMeshInput>> Targets;
 
 	UPROPERTY()
-	TObjectPtr<URecomputeUVsToolProperties> Settings = nullptr;
+	TObjectPtr<UUVEditorRecomputeUVsToolProperties> Settings = nullptr;
 
 	UPROPERTY()
 	TObjectPtr<UPolygroupLayersProperties> PolygroupLayerProperties = nullptr;
 
 	UPROPERTY()
-	TArray<TObjectPtr<URecomputeUVsOpFactory>> Factories;
+	TArray<TObjectPtr<UUVEditorRecomputeUVsOpFactory>> Factories;
 
 	TSharedPtr<UE::Geometry::FPolygroupSet, ESPMode::ThreadSafe> ActiveGroupSet;
 	void OnSelectedGroupLayerChanged();
diff --git a/Engine/Plugins/Experimental/UVEditor/Source/UVEditorTools/Public/Utilities/MeshUDIMClassifier.h b/Engine/Plugins/Experimental/UVEditor/Source/UVEditorTools/Public/Utilities/MeshUDIMClassifier.h
--- a/Engine/Plugins/Experimental/UVEditor/Source/UVEditorTools/Public/Utilities/MeshUDIMClassifier.h
+++ b/Engine/Plugins/Experimental/UVEditor/Source/UVEditorTools/Public/Utilities/MeshUDIMClassifier.h
@@ -20,7 +20,7 @@ namespace UE
 		class UVEDITORTOOLS_API FDynamicMeshUDIMClassifier
 		{
 		public:
-			explicit FDynamicMeshUDIMClassifier(FDynamicMeshUVOverlay* UVOverlay);
+			explicit FDynamicMeshUDIMClassifier(const FDynamicMeshUVOverlay* UVOverlay);
 
 			TArray<FVector2i> ActiveTiles() const;
 			TArray<int32> TidsForTile(FVector2i TileIndexIn) const;
@@ -30,7 +30,7 @@ namespace UE
 			void ClassifyUDIMs();
 
 			/** The UV Overlay to analyze for UDIMs */
-			FDynamicMeshUVOverlay* UVOverlay = nullptr;
+			const FDynamicMeshUVOverlay* UVOverlay = nullptr;
 
 			TMap<FVector2i, TArray<int32> > UDIMs;
 		};
diff --git a/Engine/Plugins/Runtime/GeometryProcessing/Source/DynamicMesh/Private/Parameterization/DynamicMeshUVEditor.cpp b/Engine/Plugins/Runtime/GeometryProcessing/Source/DynamicMesh/Private/Parameterization/DynamicMeshUVEditor.cpp
--- a/Engine/Plugins/Runtime/GeometryProcessing/Source/DynamicMesh/Private/Parameterization/DynamicMeshUVEditor.cpp
+++ b/Engine/Plugins/Runtime/GeometryProcessing/Source/DynamicMesh/Private/Parameterization/DynamicMeshUVEditor.cpp
@@ -1419,6 +1419,63 @@ bool FDynamicMeshUVEditor::QuickPack(int32 TargetTextureResolution, float Gutter
 	return bOK;
 }
 
+bool FDynamicMeshUVEditor::UDIMPack(int32 TargetTextureResolution, float GutterSize, const FVector2i& UDIMCoordsIn, const TArray<int32>* Triangles)
+{
+	TUniquePtr<TArray<int32>> TileTids;
+	if (Triangles)
+	{
+		TileTids = MakeUnique<TArray<int32>>(*Triangles);
+	}
+	else
+	{
+		// Add all set UV triangles
+		TileTids = MakeUnique<TArray<int32>>();
+		TileTids->Reserve(Mesh->TriangleCount());
+		for (int32 TriangleID : Mesh->TriangleIndicesItr())
+		{
+			TileTids->Add(TriangleID);
+		}
+	}
+
+	// Do this first, so we don't need to keep the TileTids around after moving it into the packer.
+	TSet<int32> ElementsToMove;
+	ElementsToMove.Reserve(TileTids->Num() * 3);
+	for (int Tid : *TileTids)
+	{
+		// If the triangle is unset, we will skip it here, before moving on.
+		if (UVOverlay->IsSetTriangle(Tid))
+		{
+			FIndex3i Elements = UVOverlay->GetTriangle(Tid);
+			ElementsToMove.Add(Elements[0]);
+			ElementsToMove.Add(Elements[1]);
+			ElementsToMove.Add(Elements[2]);
+		}
+	}
+	// Final check to make sure we didn't let any invalid element IDs from sneaking through.
+	// This should never happen, since we are filtering potentially unset UV triangles above
+	check(!ElementsToMove.Contains(IndexConstants::InvalidID));
+
+	// TODO: There is a second connected components call inside the packer that might be unnessessary. Could be a future optimization.
+	FDynamicMeshUVPacker Packer(UVOverlay, MoveTemp(TileTids));
+	Packer.TextureResolution = TargetTextureResolution;
+	Packer.GutterSize = GutterSize;
+	Packer.bAllowFlips = false;
+	bool bOK = Packer.StandardPack();
+
+	// Transform this to match the internal UV storage layout of negative Y
+	FVector2i TransformedUDIMCoords(UDIMCoordsIn.X, -UDIMCoordsIn.Y);
+
+	for (int32 Element : ElementsToMove)
+	{
+		FVector2f UV = UVOverlay->GetElement(Element);
+		UV = (UV)+(FVector2f)(TransformedUDIMCoords);
+		UVOverlay->SetElement(Element, UV);
+	}
+
+	return bOK;
+}
+
+
 double  FDynamicMeshUVEditor::DetermineAreaFromUVs(const FDynamicMeshUVOverlay& UVOverlay, const TArray<int32>& Triangles, FAxisAlignedBox2f* BoundingBox)
 {
 	if (BoundingBox)
diff --git a/Engine/Plugins/Runtime/GeometryProcessing/Source/DynamicMesh/Public/Parameterization/DynamicMeshUVEditor.h b/Engine/Plugins/Runtime/GeometryProcessing/Source/DynamicMesh/Public/Parameterization/DynamicMeshUVEditor.h
--- a/Engine/Plugins/Runtime/GeometryProcessing/Source/DynamicMesh/Public/Parameterization/DynamicMeshUVEditor.h
+++ b/Engine/Plugins/Runtime/GeometryProcessing/Source/DynamicMesh/Public/Parameterization/DynamicMeshUVEditor.h
@@ -231,6 +231,13 @@ public:
 	 */
 	bool QuickPack(int32 TargetTextureResolution = 512, float GutterSize = 1.0f);
 
+	/**
+    * Pack specific UVs triangles into the specific unit rectangle at the specified index
+	* @param UDIMCoordsIn Which unit rectangle tile should the UVs be packed in, indicated by the coordinate of it's lower left hand corner
+	* @param Triangles Which UV triangles to pack, all triangles if nullptr
+    */
+	bool UDIMPack(int32 TargetTextureResolution = 512, float GutterSize = 1.0f, const FVector2i& UDIMCoordsIn = FVector2i(0,0), const TArray<int32>* Triangles = nullptr);
+
 
 	//
 	// Utility functions that could maybe go elsewhere but are mainly useful in parameterization context
