UVEditor: Fix issue where assets are not being loaded initially into the UDIM display settings.

#rnx

diff --git a/Engine/Plugins/Experimental/UVEditor/Source/UVEditor/Private/UVEditorMode.cpp b/Engine/Plugins/Experimental/UVEditor/Source/UVEditor/Private/UVEditorMode.cpp
--- a/Engine/Plugins/Experimental/UVEditor/Source/UVEditor/Private/UVEditorMode.cpp
+++ b/Engine/Plugins/Experimental/UVEditor/Source/UVEditor/Private/UVEditorMode.cpp
@@ -312,6 +312,7 @@ void UUVEditorMode::Enter()
 
 	UVEditorUDIMProperties = NewObject< UUVEditorUDIMProperties >(this);
 	UVEditorUDIMProperties->Initialize(this);
+	UVEditorUDIMProperties->InitializeAssets(ToolInputObjects);
 	UDIMsChangedWatcherId = UVEditorUDIMProperties->WatchProperty(UVEditorUDIMProperties->ActiveUDIMs,
 		[this](const TArray<FUDIMSpecifier>& ActiveUDIMs) {
 			UpdateActiveUDIMs();
@@ -331,7 +332,6 @@ void UUVEditorMode::Enter()
 			}
 			return false;
 		});
-
 	PropertyObjectsToTick.Add(UVEditorUDIMProperties);
 
 	RegisterTools();
@@ -894,10 +894,7 @@ void UUVEditorMode::InitializeTargets()
 	// Prep things for layer/channel selection
 	InitializeAssetNames(ToolTargets, AssetNames);
 	PendingUVLayerIndex.SetNumZeroed(ToolTargets.Num());
-	if (UVEditorUDIMProperties)
-	{
-		UVEditorUDIMProperties->InitializeAssets(ToolInputObjects);
-	}
+
 
 	// Finish initializing the selection api
 	SelectionAPI->SetTargets(ToolInputObjects);
