UVEditor: Fix NonUnity compiler error in UVEditorBackgroundPreview.cpp

diff --git a/Engine/Plugins/Experimental/UVEditor/Source/UVEditor/Private/UVEditorBackgroundPreview.cpp b/Engine/Plugins/Experimental/UVEditor/Source/UVEditor/Private/UVEditorBackgroundPreview.cpp
--- a/Engine/Plugins/Experimental/UVEditor/Source/UVEditor/Private/UVEditorBackgroundPreview.cpp
+++ b/Engine/Plugins/Experimental/UVEditor/Source/UVEditor/Private/UVEditorBackgroundPreview.cpp
@@ -3,8 +3,8 @@
 #include "UVEditorBackgroundPreview.h"
 #include "DynamicMesh/DynamicMesh3.h"
 #include "DynamicMesh/DynamicMeshAttributeSet.h"
-#include "Drawing/LineSetComponent.h"
 #include "Drawing/MeshWireframeComponent.h"
+#include "Engine/Texture2D.h"
 #include "ToolSetupUtil.h"
 #include "Async/Async.h"
 #include "UVEditorUXSettings.h"
