UVEditor: Allow selection mode changes even if the mechanic is not currently enabled.

This adjustment is designed to support tools where the selection type before the tool is activated is not compatible with the tool's function and requires conversion (e.g. vertices to triangles). The old behavior was also partially broken, where attempts to call a mode change on an inactive selection mechanic would result in the mode being changed transparently, but with no other processing taken, such as converting the active selection or sending undo changes.

diff --git a/Engine/Plugins/Experimental/UVEditor/Source/UVEditorTools/Private/Selection/UVEditorMeshSelectionMechanic.cpp b/Engine/Plugins/Experimental/UVEditor/Source/UVEditorTools/Private/Selection/UVEditorMeshSelectionMechanic.cpp
--- a/Engine/Plugins/Experimental/UVEditor/Source/UVEditorTools/Private/Selection/UVEditorMeshSelectionMechanic.cpp
+++ b/Engine/Plugins/Experimental/UVEditor/Source/UVEditorTools/Private/Selection/UVEditorMeshSelectionMechanic.cpp
@@ -571,11 +571,11 @@ void UUVEditorMeshSelectionMechanic::SetSelectionMode(
 	using namespace UVEditorMeshSelectionMechanicLocals;
 
 	ESelectionMode OldMode = SelectionMode;
-	SelectionMode = TargetMode;
-	if (OldMode == SelectionMode || !bIsEnabled || !SelectionAPI)
+	if (OldMode == TargetMode || !SelectionAPI)
 	{
 		return;
 	}
+	SelectionMode = TargetMode;
 
 	if (ViewportButtonsAPI)
 	{
