UVEditor: Enable selection support within the Layout tool

diff --git a/Engine/Plugins/Experimental/UVEditor/Source/UVEditorTools/Private/Operators/UVEditorUVLayoutOp.cpp b/Engine/Plugins/Experimental/UVEditor/Source/UVEditorTools/Private/Operators/UVEditorUVLayoutOp.cpp
--- a/Engine/Plugins/Experimental/UVEditor/Source/UVEditorTools/Private/Operators/UVEditorUVLayoutOp.cpp
+++ b/Engine/Plugins/Experimental/UVEditor/Source/UVEditorTools/Private/Operators/UVEditorUVLayoutOp.cpp
@@ -62,7 +62,19 @@ void FUVEditorUVLayoutOp::CalculateResult(FProgressCancel* Progress)
 
 		for (FVector2i TileIndex : Tiles)
 		{
-			TUniquePtr<TArray<int32>> TileTids = MakeUnique<TArray<int32>>(TileClassifier.TidsForTile(TileIndex));
+			TUniquePtr<TArray<int32>> TileTids;
+			if (Selection.IsSet())
+			{
+				TSet<int32>& SelectionSet = Selection.GetValue();
+				TileTids = MakeUnique<TArray<int32>>(TileClassifier.TidsForTile(TileIndex).FilterByPredicate([&SelectionSet](int32 Tid)
+					{
+						return SelectionSet.Contains(Tid);
+					}));
+			}
+			else
+			{
+				TileTids = MakeUnique<TArray<int32>>(TileClassifier.TidsForTile(TileIndex));
+			}
 
 			// Do this first, so we don't need to keep the TileTids around after moving it into the packer.
 			TSet<int32> ElementsToMove;
@@ -93,7 +105,22 @@ void FUVEditorUVLayoutOp::CalculateResult(FProgressCancel* Progress)
 	}
 	else
 	{
-		FDynamicMeshUVPacker Packer(UseUVLayer);
+		TUniquePtr<TArray<int32>> TidsToLayout = nullptr;
+		if (Selection.IsSet())
+		{
+			TidsToLayout = MakeUnique<TArray<int32>>();
+			TSet<int32>& SelectionSet = Selection.GetValue();
+			TidsToLayout->Reserve(SelectionSet.Num());
+			for(int Tid : SelectionSet)
+			{
+				if (OriginalMesh->IsTriangle(Tid))
+				{
+					TidsToLayout->Add(Tid);
+				}
+			}
+		}
+		
+		FDynamicMeshUVPacker Packer(UseUVLayer, MoveTemp(TidsToLayout) );
 		ExecutePacker(Packer);
 		if (Progress && Progress->Cancelled())
 		{
@@ -163,6 +190,7 @@ TUniquePtr<FDynamicMeshOperator> UUVEditorUVLayoutOperatorFactory::MakeNewOperat
 	Op->UVTranslation = FVector2f(Settings->Translation);
 	Op->SetTransform(TargetTransform);
 	Op->bMaintainOriginatingUDIM = Settings->bEnableUDIMLayout;
+	Op->Selection = Selection;
 
 	return Op;
 }
\ No newline at end of file
diff --git a/Engine/Plugins/Experimental/UVEditor/Source/UVEditorTools/Private/UVEditorLayoutTool.cpp b/Engine/Plugins/Experimental/UVEditor/Source/UVEditorTools/Private/UVEditorLayoutTool.cpp
--- a/Engine/Plugins/Experimental/UVEditor/Source/UVEditorTools/Private/UVEditorLayoutTool.cpp
+++ b/Engine/Plugins/Experimental/UVEditor/Source/UVEditorTools/Private/UVEditorLayoutTool.cpp
@@ -48,23 +48,57 @@ void UUVEditorLayoutTool::Setup()
 	Settings->bUDIMCVAREnabled = (FUVEditorUXSettings::CVarEnablePrototypeUDIMSupport.GetValueOnGameThread() > 0);
 	AddToolPropertySource(Settings);
 
+	UContextObjectStore* ContextStore = GetToolManager()->GetContextObjectStore();
+	UVToolSelectionAPI = ContextStore->FindContext<UUVToolSelectionAPI>();
 
-	Factories.SetNum(Targets.Num());
-	for (int32 TargetIndex = 0; TargetIndex < Targets.Num(); ++TargetIndex)
+	UUVToolSelectionAPI::FHighlightOptions HighlightOptions;
+	HighlightOptions.bBaseHighlightOnPreviews = true;
+	HighlightOptions.bAutoUpdateUnwrap = true;
+	UVToolSelectionAPI->SetSelectionMechanicMode(UUVToolSelectionAPI::EUVEditorSelectionMode::Triangle);
+	UVToolSelectionAPI->SetHighlightOptions(HighlightOptions);
+	UVToolSelectionAPI->SetHighlightVisible(true, false, true);
+
+	auto SetupOpFactory = [this](UUVEditorToolMeshInput& Target, TSet<int32>* Selection)
 	{
-		TObjectPtr<UUVEditorToolMeshInput> Target = Targets[TargetIndex];
-		Factories[TargetIndex] = NewObject<UUVEditorUVLayoutOperatorFactory>();
-		Factories[TargetIndex]->TargetTransform = Target->AppliedPreview->PreviewMesh->GetTransform();
-		Factories[TargetIndex]->Settings = Settings;
-		Factories[TargetIndex]->OriginalMesh = Target->AppliedCanonical;
-		Factories[TargetIndex]->GetSelectedUVChannel = [Target]() { return Target->UVLayerIndex; };
-
-		Target->AppliedPreview->ChangeOpFactory(Factories[TargetIndex]);
-		Target->AppliedPreview->OnMeshUpdated.AddWeakLambda(this, [Target](UMeshOpPreviewWithBackgroundCompute* Preview) {
-			Target->UpdateUnwrapPreviewFromAppliedPreview();
-		});
+		TObjectPtr<UUVEditorUVLayoutOperatorFactory> Factory = NewObject<UUVEditorUVLayoutOperatorFactory>();
+		Factory->TargetTransform = Target.AppliedPreview->PreviewMesh->GetTransform();
+		Factory->Settings = Settings;
+		Factory->OriginalMesh = Target.AppliedCanonical;
+		Factory->GetSelectedUVChannel = [&Target]() { return Target.UVLayerIndex; };
+		if (Selection)
+		{
+			Factory->Selection.Emplace(*Selection);
+		}
 
-		Target->AppliedPreview->InvalidateResult();
+		Target.AppliedPreview->ChangeOpFactory(Factory);
+		Target.AppliedPreview->OnMeshUpdated.AddWeakLambda(this, [this, &Target](UMeshOpPreviewWithBackgroundCompute* Preview) {
+			Target.UpdateUnwrapPreviewFromAppliedPreview();
+
+			this->UVToolSelectionAPI->RebuildUnwrapHighlight(Preview->PreviewMesh->GetTransform());
+			});
+
+		Target.AppliedPreview->InvalidateResult();
+		return Factory;
+	};
+
+	if (UVToolSelectionAPI->HaveSelections())
+	{
+		Factories.Reserve(UVToolSelectionAPI->GetSelections().Num());
+		for (FUVToolSelection Selection : UVToolSelectionAPI->GetSelections())
+		{
+			if (ensure(Selection.Type == FUVToolSelection::EType::Triangle))
+			{
+				Factories.Add(SetupOpFactory(*Selection.Target, &Selection.SelectedIDs));
+			}
+		}
+	}
+	else
+	{
+		Factories.Reserve(Targets.Num());
+		for (int32 TargetIndex = 0; TargetIndex < Targets.Num(); ++TargetIndex)
+		{
+			Factories.Add(SetupOpFactory(*Targets[TargetIndex], nullptr));
+		}
 	}
 
 	SetToolDisplayName(LOCTEXT("ToolName", "UV Layout"));
diff --git a/Engine/Plugins/Experimental/UVEditor/Source/UVEditorTools/Public/Operators/UVEditorUVLayoutOp.h b/Engine/Plugins/Experimental/UVEditor/Source/UVEditorTools/Public/Operators/UVEditorUVLayoutOp.h
--- a/Engine/Plugins/Experimental/UVEditor/Source/UVEditorTools/Public/Operators/UVEditorUVLayoutOp.h
+++ b/Engine/Plugins/Experimental/UVEditor/Source/UVEditorTools/Public/Operators/UVEditorUVLayoutOp.h
@@ -99,6 +99,7 @@ public:
 	float UVScaleFactor = 1.0;
 	float GutterSize = 1.0;
 	bool bMaintainOriginatingUDIM = false;
+	TOptional<TSet<int32>> Selection;
 	FVector2f UVTranslation = FVector2f::Zero();
 
 	void SetTransform(const FTransformSRT3d& Transform);
@@ -134,6 +135,8 @@ public:
 	UPROPERTY()
 	TObjectPtr<UUVEditorUVLayoutProperties> Settings;
 
+	TOptional<TSet<int32>> Selection;
+
 	TSharedPtr<UE::Geometry::FDynamicMesh3, ESPMode::ThreadSafe> OriginalMesh;
 	TUniqueFunction<int32()> GetSelectedUVChannel = []() { return 0; };
 	FTransform TargetTransform;
diff --git a/Engine/Plugins/Experimental/UVEditor/Source/UVEditorTools/Public/UVEditorLayoutTool.h b/Engine/Plugins/Experimental/UVEditor/Source/UVEditorTools/Public/UVEditorLayoutTool.h
--- a/Engine/Plugins/Experimental/UVEditor/Source/UVEditorTools/Public/UVEditorLayoutTool.h
+++ b/Engine/Plugins/Experimental/UVEditor/Source/UVEditorTools/Public/UVEditorLayoutTool.h
@@ -7,6 +7,7 @@
 #include "InteractiveTool.h"
 #include "InteractiveToolBuilder.h"
 #include "UVEditorToolAnalyticsUtils.h"
+#include "Selection/UVToolSelectionAPI.h"
 
 #include "UVEditorLayoutTool.generated.h"
 
@@ -32,7 +33,7 @@ public:
  * 
  */
 UCLASS()
-class UVEDITORTOOLS_API UUVEditorLayoutTool : public UInteractiveTool
+class UVEDITORTOOLS_API UUVEditorLayoutTool : public UInteractiveTool, public IUVToolSupportsSelection
 {
 	GENERATED_BODY()
 
@@ -67,6 +68,9 @@ protected:
 	UPROPERTY()
 	TArray<TObjectPtr<UUVEditorUVLayoutOperatorFactory>> Factories;
 
+	UPROPERTY()
+	TObjectPtr<UUVToolSelectionAPI> UVToolSelectionAPI = nullptr;
+
 	//
 	// Analytics
 	//
