UVEditor: Enable support of selection within the Unwrap tool

diff --git a/Engine/Plugins/Experimental/UVEditor/Source/UVEditorTools/Private/Operators/UVEditorRecomputeUVsOp.cpp b/Engine/Plugins/Experimental/UVEditor/Source/UVEditorTools/Private/Operators/UVEditorRecomputeUVsOp.cpp
--- a/Engine/Plugins/Experimental/UVEditor/Source/UVEditorTools/Private/Operators/UVEditorRecomputeUVsOp.cpp
+++ b/Engine/Plugins/Experimental/UVEditor/Source/UVEditorTools/Private/Operators/UVEditorRecomputeUVsOp.cpp
@@ -91,15 +91,18 @@ void FUVEditorRecomputeUVsOp::NormalizeUVAreas(const FDynamicMesh3& Mesh, FDynam
 void FUVEditorRecomputeUVsOp::CalculateResult(FProgressCancel* Progress)
 {
 	NewResultInfo = FGeometryResult(EGeometryResultType::InProgress);
+	bool bOK = false;
 
-	bool bOK;
-	if (bMergingOptimization)
-	{
-		bOK = CalculateResult_RegionOptimization(Progress);
-	}
-	else
-	{
-		bOK = CalculateResult_Basic(Progress);
+	if (IsValid())
+	{		
+		if (bMergingOptimization)
+		{
+			bOK = CalculateResult_RegionOptimization(Progress);
+		}
+		else
+		{
+			bOK = CalculateResult_Basic(Progress);
+		}
 	}
 
 	NewResultInfo.SetSuccess(bOK, Progress);
@@ -107,7 +110,6 @@ void FUVEditorRecomputeUVsOp::CalculateResult(FProgressCancel* Progress)
 
 }
 
-
 void FUVEditorRecomputeUVsOp::CollectIslandComponentsPerTile(const FDynamicMeshUVOverlay& UVOverlay, TArray< TileConnectedComponents >& ComponentsPerTile, bool& bUseExistingUVTopology)
 {
 	bUseExistingUVTopology = false;
@@ -116,10 +118,15 @@ void FUVEditorRecomputeUVsOp::CollectIslandComponentsPerTile(const FDynamicMeshU
 	TArray<FVector2i> Tiles;
 	TArray<TUniquePtr<TArray<int32>>> TileTids;
 
-	FDynamicMeshUDIMClassifier TileClassifier(&UVOverlay);
-
 	if (bUDIMsEnabled)
 	{
+		TOptional<TArray<int32>> SelectionArray;
+		if (Selection.IsSet())
+		{
+			SelectionArray = Selection.GetValue().Array();
+		}		
+		FDynamicMeshUDIMClassifier TileClassifier(&UVOverlay, SelectionArray);
+	
 		Tiles = TileClassifier.ActiveTiles();
 		for (const FVector2i& Tile : Tiles)
 		{
@@ -128,8 +135,16 @@ void FUVEditorRecomputeUVsOp::CollectIslandComponentsPerTile(const FDynamicMeshU
 	}
 	else
 	{
-		Tiles.Add({ 0,0 });
-		TileTids.Add(nullptr);
+		if (Selection.IsSet())
+		{
+			Tiles.Add({ 0,0 });
+			TileTids.Emplace(MakeUnique<TArray<int32>>(Selection.GetValue().Array()));
+		}
+		else
+		{
+			Tiles.Add({ 0,0 });
+			TileTids.Add(nullptr);
+		}
 	}
 	
 	for (int32 TileIndex = 0; TileIndex < Tiles.Num(); ++TileIndex)
@@ -228,12 +243,30 @@ bool FUVEditorRecomputeUVsOp::CalculateResult_Basic(FProgressCancel* Progress)
 		TArray<bool> bComponentSolved;
 		bComponentSolved.Init(false, NumComponents);
 		int32 SuccessCount = 0;
+		TArray<FAxisAlignedBox2f> PerComponentBoundingBoxes;
+		PerComponentBoundingBoxes.SetNum(NumComponents);
 
 		// TODO: the solves here could be done in parallel if we pre-allocated the island element IDs
 		for (int32 k = 0; k < NumComponents; ++k)
 		{
 			const TArray<int32>& ComponentTris = TileComponents.ConnectedComponents[k].Indices;
 
+			FAxisAlignedBox2f& UVBounds = PerComponentBoundingBoxes[k];
+			if (bPackToOriginalBounds)
+			{
+				for (int32 tid : ComponentTris)
+				{
+					if (UseOverlay->IsSetTriangle(tid))
+					{
+						FIndex3i UVTri = UseOverlay->GetTriangle(tid);
+						FVector2f U = UseOverlay->GetElement(UVTri.A);
+						FVector2f V = UseOverlay->GetElement(UVTri.B);
+						FVector2f W = UseOverlay->GetElement(UVTri.C);
+						UVBounds.Contain(U); UVBounds.Contain(V); UVBounds.Contain(W);
+					}
+				}
+			}
+
 			bComponentSolved[k] = false;
 			switch (UnwrapType)
 			{
@@ -270,17 +303,38 @@ bool FUVEditorRecomputeUVsOp::CalculateResult_Basic(FProgressCancel* Progress)
 		if (bAutoRotate)
 		{
 			ParallelFor(NumComponents, [&](int32 k)
+			{
+				if (Progress && Progress->Cancelled())
 				{
-					if (Progress && Progress->Cancelled())
-					{
-						return;
-					}
-					if (bComponentSolved[k])
-					{
-						const TArray<int32>& ComponentTris = TileComponents.ConnectedComponents[k].Indices;
-						UVEditor.AutoOrientUVArea(ComponentTris);
-					};
-				});
+					return;
+				}
+				if (bComponentSolved[k])
+				{
+					const TArray<int32>& ComponentTris = TileComponents.ConnectedComponents[k].Indices;
+					UVEditor.AutoOrientUVArea(ComponentTris);
+				};
+			});
+		}
+		if (Progress && Progress->Cancelled())
+		{
+			return false;
+		}
+
+		if (bPackToOriginalBounds)
+		{
+			ParallelFor(NumComponents, [&](int32 k)
+			{
+				if (Progress && Progress->Cancelled())
+				{
+					return;
+				}
+				if (bComponentSolved[k])
+				{
+					const TArray<int32>& ComponentTris = TileComponents.ConnectedComponents[k].Indices;
+					const FAxisAlignedBox2f& UVBounds = PerComponentBoundingBoxes[k];
+					UVEditor.ScaleUVAreaToBoundingBox(ComponentTris, UVBounds, true, true);
+				}
+			});
 		}
 		if (Progress && Progress->Cancelled())
 		{
@@ -298,7 +352,7 @@ bool FUVEditorRecomputeUVsOp::CalculateResult_Basic(FProgressCancel* Progress)
 			return false;
 		}
 
-		if (bPackUVs)
+		if (bPackUVs && ensure(TileComponents.TileTids.Num() > 0))
 		{
 			bool bPackingSuccess = UVEditor.UDIMPack(PackingTextureResolution, PackingGutterWidth, TileComponents.Tile, &TileComponents.TileTids);
 			if (!bPackingSuccess)
@@ -440,6 +494,17 @@ bool FUVEditorRecomputeUVsOp::CalculateResult_RegionOptimization(FProgressCancel
 	return true;
 }
 
+bool FUVEditorRecomputeUVsOp::IsValid() const
+{
+	if (bPackToOriginalBounds && (bMergingOptimization || IslandMode == EUVEditorRecomputeUVsIslandMode::PolyGroups))
+	{
+		// It is not a valid state to be packing into original UV island bounds if islands are being merged or re-generated from polygroups
+		return false;
+	}
+
+	return true;
+}
+
 
 /**
  * Factory
@@ -491,6 +556,11 @@ TUniquePtr<FDynamicMeshOperator> UUVEditorRecomputeUVsOpFactory::MakeNewOperator
 		RecomputeUVsOp->bPackUVs = true;
 		RecomputeUVsOp->PackingTextureResolution = Settings->TextureResolution;
 		break;
+	case EUVEditorRecomputeUVsPropertiesLayoutType::NormalizeToExistingBounds:
+		RecomputeUVsOp->bPackToOriginalBounds = true;
+		RecomputeUVsOp->bPackUVs = false;
+		RecomputeUVsOp->bNormalizeAreas = false;
+		break;
 	case EUVEditorRecomputeUVsPropertiesLayoutType::NormalizeToBounds:
 		RecomputeUVsOp->bNormalizeAreas = true;
 		RecomputeUVsOp->AreaScaling = Settings->NormalizeScale / OriginalMesh->GetBounds().MaxDim();
@@ -501,6 +571,7 @@ TUniquePtr<FDynamicMeshOperator> UUVEditorRecomputeUVsOpFactory::MakeNewOperator
 	}
 
 	RecomputeUVsOp->SetTransform(TargetTransform);
+	RecomputeUVsOp->Selection = Selection;
 
 	return RecomputeUVsOp;
 }
diff --git a/Engine/Plugins/Experimental/UVEditor/Source/UVEditorTools/Private/UVEditorRecomputeUVsTool.cpp b/Engine/Plugins/Experimental/UVEditor/Source/UVEditorTools/Private/UVEditorRecomputeUVsTool.cpp
--- a/Engine/Plugins/Experimental/UVEditor/Source/UVEditorTools/Private/UVEditorRecomputeUVsTool.cpp
+++ b/Engine/Plugins/Experimental/UVEditor/Source/UVEditorTools/Private/UVEditorRecomputeUVsTool.cpp
@@ -56,6 +56,17 @@ void UUVEditorRecomputeUVsTool::Setup()
 	AddToolPropertySource(Settings);
 	Factories.SetNum(Targets.Num());
 
+	UContextObjectStore* ContextStore = GetToolManager()->GetContextObjectStore();
+	UVToolSelectionAPI = ContextStore->FindContext<UUVToolSelectionAPI>();
+
+	UUVToolSelectionAPI::FHighlightOptions HighlightOptions;
+	HighlightOptions.bBaseHighlightOnPreviews = true;
+	HighlightOptions.bAutoUpdateUnwrap = true;
+	UVToolSelectionAPI->SetSelectionMechanicMode(UUVToolSelectionAPI::EUVEditorSelectionMode::Triangle);
+	UVToolSelectionAPI->SetHighlightOptions(HighlightOptions);
+	UVToolSelectionAPI->SetHighlightVisible(true, false, true);
+
+
 	if (Targets.Num() == 1)
 	{
 		PolygroupLayerProperties = NewObject<UPolygroupLayersProperties>(this);
@@ -71,24 +82,54 @@ void UUVEditorRecomputeUVsTool::Setup()
 		Settings->IslandGeneration = EUVEditorRecomputeUVsPropertiesIslandMode::ExistingUVs;
 	}
 
-	
-	for (int32 TargetIndex = 0; TargetIndex < Targets.Num(); ++TargetIndex)
+
+	auto SetupOpFactory = [this](UUVEditorToolMeshInput& Target, TSet<int32>* Selection)
 	{
-		TObjectPtr<UUVEditorToolMeshInput> Target = Targets[TargetIndex];
-		Factories[TargetIndex] = NewObject<UUVEditorRecomputeUVsOpFactory>();
-		Factories[TargetIndex]->TargetTransform = (FTransform3d)Target->AppliedPreview->PreviewMesh->GetTransform();
-		Factories[TargetIndex]->Settings = Settings;
-		Factories[TargetIndex]->OriginalMesh = Target->AppliedCanonical;
-		Factories[TargetIndex]->InputGroups = ActiveGroupSet;
-		Factories[TargetIndex]->GetSelectedUVChannel = [Target]() { return Target->UVLayerIndex; };
-
-		Target->AppliedPreview->ChangeOpFactory(Factories[TargetIndex]);
-		Target->AppliedPreview->OnMeshUpdated.AddWeakLambda(this, [Target](UMeshOpPreviewWithBackgroundCompute* Preview) {
-			Target->UpdateUnwrapPreviewFromAppliedPreview();
-			});
+		TObjectPtr<UUVEditorRecomputeUVsOpFactory> Factory = NewObject<UUVEditorRecomputeUVsOpFactory>();
+		Factory->TargetTransform = (FTransform3d)Target.AppliedPreview->PreviewMesh->GetTransform();
+		Factory->Settings = Settings;
+		Factory->OriginalMesh = Target.AppliedCanonical;
+		Factory->InputGroups = ActiveGroupSet;
+		Factory->GetSelectedUVChannel = [&Target]() { return Target.UVLayerIndex; };
+		if (Selection)
+		{
+			Factory->Selection.Emplace(*Selection);
+		}
 
-		Target->AppliedPreview->InvalidateResult();
+		Target.AppliedPreview->ChangeOpFactory(Factory);
+		Target.AppliedPreview->OnMeshUpdated.AddWeakLambda(this, [this, &Target](UMeshOpPreviewWithBackgroundCompute* Preview) {
+			if (Preview->HaveValidNonEmptyResult())
+			{
+				Target.UpdateUnwrapPreviewFromAppliedPreview();
+
+				this->UVToolSelectionAPI->RebuildUnwrapHighlight(Preview->PreviewMesh->GetTransform());
+			}
+			});		
+
+		Target.AppliedPreview->InvalidateResult();
+		return Factory;
+	};
+
+	if (UVToolSelectionAPI->HaveSelections())
+	{
+		Factories.Reserve(UVToolSelectionAPI->GetSelections().Num());
+		for (FUVToolSelection Selection : UVToolSelectionAPI->GetSelections())
+		{
+			if (ensure(Selection.Type == FUVToolSelection::EType::Triangle))
+			{
+				Factories.Add(SetupOpFactory(*Selection.Target, &Selection.SelectedIDs));
+			}
+		}
 	}
+	else
+	{
+		Factories.Reserve(Targets.Num());
+		for (int32 TargetIndex = 0; TargetIndex < Targets.Num(); ++TargetIndex)
+		{
+			Factories.Add(SetupOpFactory(*Targets[TargetIndex], nullptr));
+		}
+	}
+	
 
 	SetToolDisplayName(LOCTEXT("ToolNameLocal", "UV Unwrap"));
 	GetToolManager()->DisplayMessage(
diff --git a/Engine/Plugins/Experimental/UVEditor/Source/UVEditorTools/Private/Utilities/MeshUDIMClassifier.cpp b/Engine/Plugins/Experimental/UVEditor/Source/UVEditorTools/Private/Utilities/MeshUDIMClassifier.cpp
--- a/Engine/Plugins/Experimental/UVEditor/Source/UVEditorTools/Private/Utilities/MeshUDIMClassifier.cpp
+++ b/Engine/Plugins/Experimental/UVEditor/Source/UVEditorTools/Private/Utilities/MeshUDIMClassifier.cpp
@@ -7,9 +7,10 @@
 
 using namespace UE::Geometry;
 
-FDynamicMeshUDIMClassifier::FDynamicMeshUDIMClassifier(const FDynamicMeshUVOverlay* UVOverlayIn)
+FDynamicMeshUDIMClassifier::FDynamicMeshUDIMClassifier(const FDynamicMeshUVOverlay* UVOverlayIn, TOptional<TArray<int32>> SelectionIn)
 {
 	UVOverlay = UVOverlayIn;
+	Selection = SelectionIn;
 	ClassifyUDIMs();
 }
 
@@ -43,7 +44,14 @@ void FDynamicMeshUDIMClassifier::ClassifyUDIMs()
 	};
 
 	FMeshConnectedComponents UVComponents(Mesh);
-	UVComponents.FindConnectedTriangles(UVIslandPredicate);
+	if (Selection.IsSet())
+	{
+		UVComponents.FindConnectedTriangles(Selection.GetValue(), UVIslandPredicate);
+	}
+	else
+	{
+		UVComponents.FindConnectedTriangles(UVIslandPredicate);
+	}
 
 	auto UVTriangleToUDIM = [this](int32 Tid)
 	{
diff --git a/Engine/Plugins/Experimental/UVEditor/Source/UVEditorTools/Public/Operators/UVEditorRecomputeUVsOp.h b/Engine/Plugins/Experimental/UVEditor/Source/UVEditorTools/Public/Operators/UVEditorRecomputeUVsOp.h
--- a/Engine/Plugins/Experimental/UVEditor/Source/UVEditorTools/Public/Operators/UVEditorRecomputeUVsOp.h
+++ b/Engine/Plugins/Experimental/UVEditor/Source/UVEditorTools/Public/Operators/UVEditorRecomputeUVsOp.h
@@ -63,6 +63,8 @@ enum class EUVEditorRecomputeUVsPropertiesLayoutType
 	None,
 	/** Uniformly scale and translate UV islands collectively to pack them into the unit square, i.e. fit between 0 and 1 with no overlap */
 	Repack,
+	/** Scale and center all islands to fit within their original bounding boxes. Only applicable if using existing UVs. */
+	NormalizeToExistingBounds,
 	/** Uniformly scale UV islands such that they have constant relative area, relative to object bounds */
 	NormalizeToBounds,
 	/** Uniformly scale UV islands such that they have constant relative area, relative to world space */
@@ -174,6 +176,9 @@ public:
 	// orientation control
 	bool bAutoRotate = true;
 
+	// Individual packing
+	bool bPackToOriginalBounds = false;
+
 	// area scaling
 	bool bNormalizeAreas = true;
 	float AreaScaling = 1.0;
@@ -182,7 +187,7 @@ public:
 	int32 UVLayer = 0;
 
 	// Atlas Packing parameters
-	bool bPackUVs = true;
+	bool bPackUVs = true;	
 	int32 PackingTextureResolution = 512;
 	float PackingGutterWidth = 1.0f;
 
@@ -206,6 +211,8 @@ public:
 	// UDIM options	
 	bool bUDIMsEnabled = false;
 
+	// Selection
+	TOptional<TSet<int32>> Selection;
 
 	// set ability on protected transform.
 	void SetTransform(const FTransformSRT3d& XForm)
@@ -222,6 +229,11 @@ public:
 
 protected:
 
+	/*
+	* Check if operation is valid to do work. It may not be valid if there are conflicting or inconsistent options passed to it from the user.
+	*/
+	bool IsValid() const;
+
 	FGeometryResult NewResultInfo;
 
 	void NormalizeUVAreas(const FDynamicMesh3& Mesh, FDynamicMeshUVOverlay* Overlay, const TileConnectedComponents& TileComponents, float GlobalScale = 1.0f);
@@ -255,6 +267,8 @@ public:
 	UPROPERTY()
 	TObjectPtr<UUVEditorRecomputeUVsToolProperties> Settings = nullptr;
 
+	TOptional<TSet<int32>> Selection;
+
 	TSharedPtr<UE::Geometry::FPolygroupSet, ESPMode::ThreadSafe> InputGroups;
 	TSharedPtr<UE::Geometry::FDynamicMesh3, ESPMode::ThreadSafe> OriginalMesh;
 	TUniqueFunction<int32()> GetSelectedUVChannel = []() { return 0; };
diff --git a/Engine/Plugins/Experimental/UVEditor/Source/UVEditorTools/Public/UVEditorRecomputeUVsTool.h b/Engine/Plugins/Experimental/UVEditor/Source/UVEditorTools/Public/UVEditorRecomputeUVsTool.h
--- a/Engine/Plugins/Experimental/UVEditor/Source/UVEditorTools/Public/UVEditorRecomputeUVsTool.h
+++ b/Engine/Plugins/Experimental/UVEditor/Source/UVEditorTools/Public/UVEditorRecomputeUVsTool.h
@@ -9,6 +9,7 @@
 #include "Polygroups/PolygroupSet.h"
 #include "Drawing/UVLayoutPreview.h"
 #include "UVEditorToolAnalyticsUtils.h"
+#include "Selection/UVToolSelectionAPI.h"
 
 #include "UVEditorRecomputeUVsTool.generated.h"
 
@@ -43,7 +44,7 @@ public:
  * UUVEditorRecomputeUVsTool Recomputes UVs based on existing segmentations of the mesh
  */
 UCLASS()
-class UVEDITORTOOLS_API UUVEditorRecomputeUVsTool : public UInteractiveTool 
+class UVEDITORTOOLS_API UUVEditorRecomputeUVsTool : public UInteractiveTool, public IUVToolSupportsSelection
 {
 	GENERATED_BODY()
 
@@ -83,6 +84,10 @@ protected:
 	UPROPERTY()
 	TArray<TObjectPtr<UUVEditorRecomputeUVsOpFactory>> Factories;
 
+	UPROPERTY()
+	TObjectPtr<UUVToolSelectionAPI> UVToolSelectionAPI = nullptr;
+
+
 	TSharedPtr<UE::Geometry::FPolygroupSet, ESPMode::ThreadSafe> ActiveGroupSet;
 	void OnSelectedGroupLayerChanged();
 	void UpdateActiveGroupLayer();
diff --git a/Engine/Plugins/Experimental/UVEditor/Source/UVEditorTools/Public/Utilities/MeshUDIMClassifier.h b/Engine/Plugins/Experimental/UVEditor/Source/UVEditorTools/Public/Utilities/MeshUDIMClassifier.h
--- a/Engine/Plugins/Experimental/UVEditor/Source/UVEditorTools/Public/Utilities/MeshUDIMClassifier.h
+++ b/Engine/Plugins/Experimental/UVEditor/Source/UVEditorTools/Public/Utilities/MeshUDIMClassifier.h
@@ -19,8 +19,8 @@ namespace UE
 		 */
 		class UVEDITORTOOLS_API FDynamicMeshUDIMClassifier
 		{
-		public:
-			explicit FDynamicMeshUDIMClassifier(const FDynamicMeshUVOverlay* UVOverlay);
+		public:			
+			explicit FDynamicMeshUDIMClassifier(const FDynamicMeshUVOverlay* UVOverlay, TOptional<TArray<int32>> Selection = TOptional<TArray<int32>>());
 
 			TArray<FVector2i> ActiveTiles() const;
 			TArray<int32> TidsForTile(FVector2i TileIndexIn) const;
@@ -31,7 +31,9 @@ namespace UE
 
 			/** The UV Overlay to analyze for UDIMs */
 			const FDynamicMeshUVOverlay* UVOverlay = nullptr;
-
+			
+			TOptional<TArray<int32>> Selection;
+			
 			TMap<FVector2i, TArray<int32> > UDIMs;
 		};
 	}
