UVEditor: Update Layout Op to use new api for TileClassifier that uses selection directly to improve performance.

diff --git a/Engine/Plugins/Experimental/UVEditor/Source/UVEditorTools/Private/Operators/UVEditorUVLayoutOp.cpp b/Engine/Plugins/Experimental/UVEditor/Source/UVEditorTools/Private/Operators/UVEditorUVLayoutOp.cpp
--- a/Engine/Plugins/Experimental/UVEditor/Source/UVEditorTools/Private/Operators/UVEditorUVLayoutOp.cpp
+++ b/Engine/Plugins/Experimental/UVEditor/Source/UVEditorTools/Private/Operators/UVEditorUVLayoutOp.cpp
@@ -56,25 +56,20 @@ void FUVEditorUVLayoutOp::CalculateResult(FProgressCancel* Progress)
 	}
 
 	if (bMaintainOriginatingUDIM)
-	{
-		FDynamicMeshUDIMClassifier TileClassifier(UseUVLayer);
+	{		
+		TOptional<TArray<int32>> SelectionArray;
+		if (Selection.IsSet())
+		{
+			SelectionArray = Selection.GetValue().Array();
+		}
+		FDynamicMeshUDIMClassifier TileClassifier(UseUVLayer, SelectionArray);
+
 		TArray<FVector2i> Tiles = TileClassifier.ActiveTiles();
 
 		for (FVector2i TileIndex : Tiles)
 		{
 			TUniquePtr<TArray<int32>> TileTids;
-			if (Selection.IsSet())
-			{
-				TSet<int32>& SelectionSet = Selection.GetValue();
-				TileTids = MakeUnique<TArray<int32>>(TileClassifier.TidsForTile(TileIndex).FilterByPredicate([&SelectionSet](int32 Tid)
-					{
-						return SelectionSet.Contains(Tid);
-					}));
-			}
-			else
-			{
-				TileTids = MakeUnique<TArray<int32>>(TileClassifier.TidsForTile(TileIndex));
-			}
+			TileTids = MakeUnique<TArray<int32>>(TileClassifier.TidsForTile(TileIndex));
 
 			// Do this first, so we don't need to keep the TileTids around after moving it into the packer.
 			TSet<int32> ElementsToMove;
