UVEditor: Fix crash bug when due to selection hover behavior not being correctly handled by the Input Router during shutdown conditions. Additionally adds a few checks to catch unexpected future behavior.

diff --git a/Engine/Plugins/Experimental/UVEditor/Source/UVEditorTools/Private/Selection/UVEditorMeshSelectionMechanic.cpp b/Engine/Plugins/Experimental/UVEditor/Source/UVEditorTools/Private/Selection/UVEditorMeshSelectionMechanic.cpp
--- a/Engine/Plugins/Experimental/UVEditor/Source/UVEditorTools/Private/Selection/UVEditorMeshSelectionMechanic.cpp
+++ b/Engine/Plugins/Experimental/UVEditor/Source/UVEditorTools/Private/Selection/UVEditorMeshSelectionMechanic.cpp
@@ -1160,6 +1160,11 @@ bool UUVEditorMeshSelectionMechanic::OnUpdateHover(const FInputDeviceRay& Device
 {
 	using namespace UVEditorMeshSelectionMechanicLocals;
 
+	if (!IsEnabled())
+	{
+		return false;
+	}
+
 	ESelectionMode Mode = SelectionMode;
 	if (SelectionMode != ESelectionMode::Vertex && SelectionMode != ESelectionMode::Edge)
 	{
@@ -1233,9 +1238,18 @@ bool UUVEditorMeshSelectionMechanic::OnUpdateHover(const FInputDeviceRay& Device
 
 void UUVEditorMeshSelectionMechanic::OnEndHover()
 {
-	HoverPointSet->Clear();
-	HoverLineSet->Clear();
-	HoverTriangleSet->Clear();
+	if (ensure(HoverPointSet.IsValid()))
+	{
+		HoverPointSet->Clear();
+	}
+	if (ensure(HoverLineSet.IsValid()))
+	{
+		HoverLineSet->Clear();
+	}
+	if (ensure(HoverTriangleSet.IsValid()))
+	{
+		HoverTriangleSet->Clear();
+	}
 }
 
 #undef LOCTEXT_NAMESPACE
diff --git a/Engine/Plugins/Experimental/UVEditor/Source/UVEditorTools/Private/Selection/UVToolSelectionAPI.cpp b/Engine/Plugins/Experimental/UVEditor/Source/UVEditorTools/Private/Selection/UVToolSelectionAPI.cpp
--- a/Engine/Plugins/Experimental/UVEditor/Source/UVEditorTools/Private/Selection/UVToolSelectionAPI.cpp
+++ b/Engine/Plugins/Experimental/UVEditor/Source/UVEditorTools/Private/Selection/UVToolSelectionAPI.cpp
@@ -82,14 +82,6 @@ void UUVToolSelectionAPI::SetTargets(const TArray<TObjectPtr<UUVEditorToolMeshIn
 
 void UUVToolSelectionAPI::Shutdown()
 {
-	HighlightMechanic->Shutdown();
-	SelectionMechanic->Shutdown();
-
-	MechanicAdapter->Shutdown(EToolShutdownType::Completed);
-	HighlightMechanic = nullptr;
-	SelectionMechanic = nullptr;
-	MechanicAdapter = nullptr;
-
 	if (UnwrapInputRouter.IsValid())
 	{
 		// Make sure that we stop any captures that our mechanics may have, then remove them from
@@ -99,6 +91,14 @@ void UUVToolSelectionAPI::Shutdown()
 		UnwrapInputRouter = nullptr;
 	}
 
+	HighlightMechanic->Shutdown();
+	SelectionMechanic->Shutdown();
+
+	MechanicAdapter->Shutdown(EToolShutdownType::Completed);
+	HighlightMechanic = nullptr;
+	SelectionMechanic = nullptr;
+	MechanicAdapter = nullptr;
+
 	for (TObjectPtr<UUVEditorToolMeshInput> Target : Targets)
 	{
 		Target->OnCanonicalModified.RemoveAll(this);
@@ -127,6 +127,13 @@ void UUVToolSelectionAPI::OnToolEnded(UInteractiveTool* DeadTool)
 	OnSelectionChanged.RemoveAll(DeadTool);
 	OnPreSelectionChange.RemoveAll(DeadTool);
 	OnDragSelectionChanged.RemoveAll(DeadTool);
+	
+	if (UnwrapInputRouter.IsValid())
+	{
+		// Make sure that we stop any captures that our mechanics may have.
+		UnwrapInputRouter->ForceTerminateSource(MechanicAdapter);
+	}
+	
 }
 
 void UUVToolSelectionAPI::SetSelections(const TArray<FUVToolSelection>& SelectionsIn, bool bBroadcast, bool bEmitChange)
