UVEditor: Added support for a DetailsCustomization to the Unwrap tool to prevent selection of the new behavior in cases where it is incompatible with other tool options.

diff --git a/Engine/Plugins/Experimental/UVEditor/Source/UVEditorToolsEditorOnly/Private/DetailsCustomizations/UVUnwrapToolCustomizations.cpp b/Engine/Plugins/Experimental/UVEditor/Source/UVEditorToolsEditorOnly/Private/DetailsCustomizations/UVUnwrapToolCustomizations.cpp
new file mode 100644
--- /dev/null
+++ b/Engine/Plugins/Experimental/UVEditor/Source/UVEditorToolsEditorOnly/Private/DetailsCustomizations/UVUnwrapToolCustomizations.cpp
@@ -0,0 +1,117 @@
+
+// Copyright Epic Games, Inc. All Rights Reserved.
+
+#include "DetailsCustomizations/UVUnwrapToolCustomizations.h"
+#include "DetailCategoryBuilder.h"
+#include "DetailLayoutBuilder.h"
+#include "DetailWidgetRow.h"
+#include "PropertyCustomizationHelpers.h"
+#include "Widgets/Input/SComboButton.h"
+#include "IDetailChildrenBuilder.h"
+#include "Internationalization/BreakIterator.h"
+#include "Framework/MultiBox/MultiBoxBuilder.h"
+#include "IContentBrowserSingleton.h"
+#include "SAssetView.h"
+
+#include "PropertyRestriction.h"
+
+#include "Operators/UVEditorRecomputeUVsOp.h"
+
+
+#define LOCTEXT_NAMESPACE "UVEditorDetailsCustomization"
+
+namespace UVEditorDetailsCustomizationLocal
+{
+
+	template<class UENUM_TYPE>
+	FPropertyAccess::Result GetPropertyValueAsEnum(const TSharedPtr<IPropertyHandle> Property, UENUM_TYPE& Value)
+	{
+		if(Property.IsValid())
+		{
+			uint8 ValueAsByte;
+			FPropertyAccess::Result Result = Property->GetValue(/*out*/ ValueAsByte);
+
+			if (Result == FPropertyAccess::Success)
+			{
+				Value = (UENUM_TYPE)ValueAsByte;
+				return FPropertyAccess::Success;
+			}
+			else
+			{
+				return FPropertyAccess::Fail;
+			}
+		}
+		return FPropertyAccess::Fail;
+	}
+
+}
+
+
+//
+// UVEditorRecomputeUVsTool
+//
+
+
+TSharedRef<IDetailCustomization> FUVEditorRecomputeUVsToolDetails::MakeInstance()
+{
+	return MakeShareable(new FUVEditorRecomputeUVsToolDetails);
+}
+
+
+void FUVEditorRecomputeUVsToolDetails::CustomizeDetails(IDetailLayoutBuilder& DetailBuilder)
+{
+	TSharedRef<IPropertyHandle> IslandGenerationModeHandle = DetailBuilder.GetProperty(GET_MEMBER_NAME_CHECKED(UUVEditorRecomputeUVsToolProperties, IslandGeneration), UUVEditorRecomputeUVsToolProperties::StaticClass());
+	TSharedRef<IPropertyHandle> UnwrapTypeHandle = DetailBuilder.GetProperty(GET_MEMBER_NAME_CHECKED(UUVEditorRecomputeUVsToolProperties, UnwrapType), UUVEditorRecomputeUVsToolProperties::StaticClass());
+	TSharedRef<IPropertyHandle> LayoutTypeHandle = DetailBuilder.GetProperty(GET_MEMBER_NAME_CHECKED(UUVEditorRecomputeUVsToolProperties, LayoutType), UUVEditorRecomputeUVsToolProperties::StaticClass());
+
+	IslandGenerationModeHandle->SetOnPropertyValueChanged(FSimpleDelegate::CreateRaw(this, &FUVEditorRecomputeUVsToolDetails::EvaluateLayoutTypeRestrictions, IslandGenerationModeHandle, UnwrapTypeHandle, LayoutTypeHandle));
+	UnwrapTypeHandle->SetOnPropertyValueChanged(FSimpleDelegate::CreateRaw(this, &FUVEditorRecomputeUVsToolDetails::EvaluateLayoutTypeRestrictions, IslandGenerationModeHandle, UnwrapTypeHandle, LayoutTypeHandle));
+
+	static FText RestrictReason = LOCTEXT("LayoutTypeRestrictionText", "Requires existing UVs be used for IslandMode and Island Merging not selected as the unwrap strategy.");
+	LayoutTypeEnumRestriction = MakeShareable(new FPropertyRestriction(RestrictReason));
+	LayoutTypeHandle->AddRestriction(LayoutTypeEnumRestriction.ToSharedRef());
+
+	EvaluateLayoutTypeRestrictions(IslandGenerationModeHandle, UnwrapTypeHandle, LayoutTypeHandle);	
+}
+
+void FUVEditorRecomputeUVsToolDetails::EvaluateLayoutTypeRestrictions(TSharedRef<IPropertyHandle> IslandGenerationModeHandle, TSharedRef<IPropertyHandle> UnwrapTypeHandle, TSharedRef<IPropertyHandle> LayoutTypeHandle)
+{
+	EUVEditorRecomputeUVsPropertiesIslandMode IslandMode = EUVEditorRecomputeUVsPropertiesIslandMode::PolyGroups;
+	EUVEditorRecomputeUVsPropertiesUnwrapType UnwrapType = EUVEditorRecomputeUVsPropertiesUnwrapType::Conformal;
+	FPropertyAccess::Result IslandModeResult = UVEditorDetailsCustomizationLocal::GetPropertyValueAsEnum(IslandGenerationModeHandle, IslandMode);
+	FPropertyAccess::Result UnwrapTypeResult = UVEditorDetailsCustomizationLocal::GetPropertyValueAsEnum(UnwrapTypeHandle, UnwrapType);
+
+	UEnum* ImportTypeEnum = StaticEnum<EUVEditorRecomputeUVsPropertiesLayoutType>();
+
+	LayoutTypeEnumRestriction->RemoveAll();
+	if (IslandModeResult != FPropertyAccess::Fail &&
+		UnwrapTypeResult != FPropertyAccess::Fail &&
+		(IslandMode == EUVEditorRecomputeUVsPropertiesIslandMode::PolyGroups ||
+		 UnwrapType == EUVEditorRecomputeUVsPropertiesUnwrapType::IslandMerging))
+	{
+		uint8 EnumValue = ImportTypeEnum->GetValueByName("EUVEditorRecomputeUVsPropertiesLayoutType::NormalizeToExistingBounds");
+		LayoutTypeEnumRestriction->AddDisabledValue("NormalizeToExistingBounds");
+
+		uint8 Value = 0;
+		LayoutTypeHandle->GetValue(Value);
+		// The goal here is to shift our selected layout type away from the offending option, either up or down one.
+		// This code assumes though that there are options other than our offending enum value to move off to. If it's the only one,
+		// this code will end up doing the wrong thing.
+		if (Value == EnumValue)
+		{
+			if (Value > 0)
+			{
+				Value = EnumValue - 1;
+			}
+			else
+			{
+				Value = EnumValue + 1;
+			}
+		}
+		LayoutTypeHandle->SetValue(Value, EPropertyValueSetFlags::NotTransactable | EPropertyValueSetFlags::InteractiveChange);
+	}
+
+
+}
+
+#undef LOCTEXT_NAMESPACE
\ No newline at end of file
diff --git a/Engine/Plugins/Experimental/UVEditor/Source/UVEditorToolsEditorOnly/Private/UVEditorToolsEditorOnlyModule.cpp b/Engine/Plugins/Experimental/UVEditor/Source/UVEditorToolsEditorOnly/Private/UVEditorToolsEditorOnlyModule.cpp
--- a/Engine/Plugins/Experimental/UVEditor/Source/UVEditorToolsEditorOnly/Private/UVEditorToolsEditorOnlyModule.cpp
+++ b/Engine/Plugins/Experimental/UVEditor/Source/UVEditorToolsEditorOnly/Private/UVEditorToolsEditorOnlyModule.cpp
@@ -2,17 +2,52 @@
 
 #include "UVEditorToolsEditorOnlyModule.h"
 
+#include "DetailsCustomizations/UVUnwrapToolCustomizations.h"
+
+#include "Operators/UVEditorRecomputeUVsOp.h"
+
 #define LOCTEXT_NAMESPACE "FUVEditorToolsEditorOnlyModule"
 
 void FUVEditorToolsEditorOnlyModule::StartupModule()
 {
 	// This code will execute after your module is loaded into memory; the exact timing is specified in the .uplugin file per-module
+
+	FCoreDelegates::OnPostEngineInit.AddRaw(this, &FUVEditorToolsEditorOnlyModule::OnPostEngineInit);
+
 }
 
 void FUVEditorToolsEditorOnlyModule::ShutdownModule()
 {
 	// This function may be called during shutdown to clean up your module.  For modules that support dynamic reloading,
 	// we call this function before unloading the module.
+
+
+	FCoreDelegates::OnPostEngineInit.RemoveAll(this);
+
+	// Unregister customizations
+	FPropertyEditorModule* PropertyEditorModule = FModuleManager::GetModulePtr<FPropertyEditorModule>("PropertyEditor");
+	if (PropertyEditorModule)
+	{
+		for (const FName& ClassName : ClassesToUnregisterOnShutdown)
+		{
+			PropertyEditorModule->UnregisterCustomClassLayout(ClassName);
+		}
+		for (const FName& PropertyName : PropertiesToUnregisterOnShutdown)
+		{
+			PropertyEditorModule->UnregisterCustomPropertyTypeLayout(PropertyName);
+		}
+	}
+}
+
+void FUVEditorToolsEditorOnlyModule::OnPostEngineInit()
+{
+
+	// Register details view customizations
+	FPropertyEditorModule& PropertyModule = FModuleManager::LoadModuleChecked<FPropertyEditorModule>("PropertyEditor");
+
+	/// Unwrap
+	PropertyModule.RegisterCustomClassLayout("UVEditorRecomputeUVsToolProperties", FOnGetDetailCustomizationInstance::CreateStatic(&FUVEditorRecomputeUVsToolDetails::MakeInstance));
+	ClassesToUnregisterOnShutdown.Add(UUVEditorRecomputeUVsToolProperties::StaticClass()->GetFName());
 }
 
 #undef LOCTEXT_NAMESPACE
diff --git a/Engine/Plugins/Experimental/UVEditor/Source/UVEditorToolsEditorOnly/Public/DetailsCustomizations/UVUnwrapToolCustomizations.h b/Engine/Plugins/Experimental/UVEditor/Source/UVEditorToolsEditorOnly/Public/DetailsCustomizations/UVUnwrapToolCustomizations.h
new file mode 100644
--- /dev/null
+++ b/Engine/Plugins/Experimental/UVEditor/Source/UVEditorToolsEditorOnly/Public/DetailsCustomizations/UVUnwrapToolCustomizations.h
@@ -0,0 +1,22 @@
+// Copyright Epic Games, Inc. All Rights Reserved.
+
+#pragma once
+
+#include "IDetailCustomization.h"
+#include "Templates/SharedPointer.h"
+
+// Customization for UUVEditorRecomputeUVsToolProperties
+class FUVEditorRecomputeUVsToolDetails : public IDetailCustomization
+{
+public:
+	static TSharedRef<IDetailCustomization> MakeInstance();
+	virtual void CustomizeDetails(IDetailLayoutBuilder& DetailBuilder) override;
+
+protected:
+	void EvaluateLayoutTypeRestrictions(TSharedRef<IPropertyHandle> IslandGenerationModeHandle, TSharedRef<IPropertyHandle> UnwrapTypeHandle, TSharedRef<IPropertyHandle> LayoutTypeHandle);
+
+	/** Property restriction applied to blend paint enum dropdown box */
+	TSharedPtr<FPropertyRestriction> LayoutTypeEnumRestriction;
+};
+
+
diff --git a/Engine/Plugins/Experimental/UVEditor/Source/UVEditorToolsEditorOnly/Public/UVEditorToolsEditorOnlyModule.h b/Engine/Plugins/Experimental/UVEditor/Source/UVEditorToolsEditorOnly/Public/UVEditorToolsEditorOnlyModule.h
--- a/Engine/Plugins/Experimental/UVEditor/Source/UVEditorToolsEditorOnly/Public/UVEditorToolsEditorOnlyModule.h
+++ b/Engine/Plugins/Experimental/UVEditor/Source/UVEditorToolsEditorOnly/Public/UVEditorToolsEditorOnlyModule.h
@@ -12,4 +12,11 @@ public:
 	/** IModuleInterface implementation */
 	virtual void StartupModule() override;
 	virtual void ShutdownModule() override;
+
+	void OnPostEngineInit();
+
+protected:
+
+	TArray<FName> ClassesToUnregisterOnShutdown;
+	TArray<FName> PropertiesToUnregisterOnShutdown;
 };
\ No newline at end of file
diff --git a/Engine/Plugins/Experimental/UVEditor/Source/UVEditorToolsEditorOnly/UVEditorToolsEditorOnly.Build.cs b/Engine/Plugins/Experimental/UVEditor/Source/UVEditorToolsEditorOnly/UVEditorToolsEditorOnly.Build.cs
--- a/Engine/Plugins/Experimental/UVEditor/Source/UVEditorToolsEditorOnly/UVEditorToolsEditorOnly.Build.cs
+++ b/Engine/Plugins/Experimental/UVEditor/Source/UVEditorToolsEditorOnly/UVEditorToolsEditorOnly.Build.cs
@@ -53,6 +53,11 @@ public class UVEditorToolsEditorOnly : ModuleRules
 				"ModelingOperators",
 				"ModelingOperatorsEditorOnly",
 				"UVEditorTools",
+
+				"ToolWidgets",
+				"EditorWidgets",
+				"ModelingEditorUI",
+				"PropertyEditor",
 				// ... add private dependencies that you statically link with here ...	
 			}
 			);
