UVEditor: Fix non-unity build issues related to the unwrap tool details customization

#rnx

diff --git a/Engine/Plugins/Experimental/UVEditor/Source/UVEditorToolsEditorOnly/Private/UVEditorToolsEditorOnlyModule.cpp b/Engine/Plugins/Experimental/UVEditor/Source/UVEditorToolsEditorOnly/Private/UVEditorToolsEditorOnlyModule.cpp
--- a/Engine/Plugins/Experimental/UVEditor/Source/UVEditorToolsEditorOnly/Private/UVEditorToolsEditorOnlyModule.cpp
+++ b/Engine/Plugins/Experimental/UVEditor/Source/UVEditorToolsEditorOnly/Private/UVEditorToolsEditorOnlyModule.cpp
@@ -2,6 +2,9 @@
 
 #include "UVEditorToolsEditorOnlyModule.h"
 
+#include "Modules/ModuleManager.h"
+#include "PropertyEditorModule.h"
+
 #include "DetailsCustomizations/UVUnwrapToolCustomizations.h"
 
 #include "Operators/UVEditorRecomputeUVsOp.h"
diff --git a/Engine/Plugins/Experimental/UVEditor/Source/UVEditorToolsEditorOnly/Public/DetailsCustomizations/UVUnwrapToolCustomizations.h b/Engine/Plugins/Experimental/UVEditor/Source/UVEditorToolsEditorOnly/Public/DetailsCustomizations/UVUnwrapToolCustomizations.h
--- a/Engine/Plugins/Experimental/UVEditor/Source/UVEditorToolsEditorOnly/Public/DetailsCustomizations/UVUnwrapToolCustomizations.h
+++ b/Engine/Plugins/Experimental/UVEditor/Source/UVEditorToolsEditorOnly/Public/DetailsCustomizations/UVUnwrapToolCustomizations.h
@@ -5,6 +5,9 @@
 #include "IDetailCustomization.h"
 #include "Templates/SharedPointer.h"
 
+class IPropertyHandle;
+class FPropertyRestriction;
+
 // Customization for UUVEditorRecomputeUVsToolProperties
 class FUVEditorRecomputeUVsToolDetails : public IDetailCustomization
 {
