UVEditor: Fix bug where initial selection mode is not functional upon editor startup until a selection mode button is manually clicked by the user.

diff --git a/Engine/Plugins/Experimental/UVEditor/Source/UVEditorTools/Private/Selection/UVEditorMeshSelectionMechanic.cpp b/Engine/Plugins/Experimental/UVEditor/Source/UVEditorTools/Private/Selection/UVEditorMeshSelectionMechanic.cpp
--- a/Engine/Plugins/Experimental/UVEditor/Source/UVEditorTools/Private/Selection/UVEditorMeshSelectionMechanic.cpp
+++ b/Engine/Plugins/Experimental/UVEditor/Source/UVEditorTools/Private/Selection/UVEditorMeshSelectionMechanic.cpp
@@ -426,12 +426,9 @@ void UUVEditorMeshSelectionMechanic::Setup(UInteractiveTool* ParentToolIn)
 		[this](UUVToolSelectionAPI::EUVEditorSelectionMode NewMode) {
 			SetSelectionMode(NewMode);
 		});
-	// Make sure we match the activated button
-	FModeChangeOptions ModeChangeOptions;
-	ModeChangeOptions.bEmitChanges = false;
-	SetSelectionMode(ViewportButtonsAPI->GetSelectionMode(), 
-		ModeChangeOptions); // convert, broadcast, don't emit
 
+	// Make sure we match the activated button
+	SelectionMode = ViewportButtonsAPI->GetSelectionMode();
 	SetIsEnabled(bIsEnabled);
 }
 
@@ -571,12 +568,18 @@ void UUVEditorMeshSelectionMechanic::SetSelectionMode(
 	using namespace UVEditorMeshSelectionMechanicLocals;
 
 	ESelectionMode OldMode = SelectionMode;
-	if (OldMode == TargetMode || !SelectionAPI)
+	if (OldMode == TargetMode)
 	{
 		return;
 	}
 	SelectionMode = TargetMode;
 
+	if (SelectionAPI == nullptr)
+	{
+		return; // From this point on, if we don't have a SelectionAPI assigned yet, we won't do anything further for this.
+		        // This might happen during initialization.
+	}
+
 	if (ViewportButtonsAPI)
 	{
 		// Not clear whether we should or shouldn't broadcast this. A user could conceivably set selection
