UVEditor: Fix issue with the background texture display incorrectly handling textures that need non-color samplers in materials, e.g. normals.

diff --git a/Engine/Plugins/Experimental/UVEditor/Source/UVEditor/Private/UVEditorBackgroundPreview.cpp b/Engine/Plugins/Experimental/UVEditor/Source/UVEditor/Private/UVEditorBackgroundPreview.cpp
--- a/Engine/Plugins/Experimental/UVEditor/Source/UVEditor/Private/UVEditorBackgroundPreview.cpp
+++ b/Engine/Plugins/Experimental/UVEditor/Source/UVEditor/Private/UVEditorBackgroundPreview.cpp
@@ -9,6 +9,7 @@
 #include "Async/Async.h"
 #include "UVEditorUXSettings.h"
 #include "UDIMUtilities.h"
+#include "Materials/MaterialExpressionTextureBase.h"
 
 using namespace UE::Geometry;
 
@@ -64,9 +65,13 @@ void UUVEditorBackgroundPreview::UpdateBackground()
 	const FVector Normal(0, 0, 1);
 	const FColor BackgroundColor = FColor::Blue;
 
-	UMaterial* Material = LoadObject<UMaterial>(nullptr, TEXT("/UVEditor/Materials/UVEditorBackground"));
-	check(Material);	
-	BackgroundMaterial = UMaterialInstanceDynamic::Create(Material, this);
+	UMaterial* RegularMaterial = LoadObject<UMaterial>(nullptr, TEXT("/UVEditor/Materials/UVEditorBackground"));
+	UMaterial* VirtualTextureMaterial = LoadObject<UMaterial>(nullptr, TEXT("/UVEditor/Materials/UVEditorVTBackground"));
+
+	check(RegularMaterial);
+	check(VirtualTextureMaterial);
+
+	BackgroundMaterial = UMaterialInstanceDynamic::Create(RegularMaterial, this);
 	switch (Settings->SourceType)
 	{
 		case EUVEditorBackgroundSourceType::Checkerboard:
@@ -90,13 +95,83 @@ void UUVEditorBackgroundPreview::UpdateBackground()
 			{
 				if (Settings->SourceTexture->IsCurrentlyVirtualTextured())
 				{
-					BackgroundMaterial->SetTextureParameterValue(TEXT("BackgroundVTBaseMap"), Settings->SourceTexture);
-					BackgroundMaterial->SetScalarParameterValue(TEXT("BackgroundVirtualTextureSwitch"), 1);
+					BackgroundMaterial = UMaterialInstanceDynamic::Create(VirtualTextureMaterial, this);
 				}
-				else
+
+				EMaterialSamplerType SamplerRequiredForTexture = UMaterialExpressionTextureBase::GetSamplerTypeForTexture(Settings->SourceTexture.Get());
+
+				switch (SamplerRequiredForTexture)
 				{
-					BackgroundMaterial->SetTextureParameterValue(TEXT("BackgroundBaseMap"), Settings->SourceTexture);
-					BackgroundMaterial->SetScalarParameterValue(TEXT("BackgroundVirtualTextureSwitch"), 0);
+				case EMaterialSamplerType::SAMPLERTYPE_Color:
+					BackgroundMaterial->SetTextureParameterValue(TEXT("BackgroundBaseMap_Color"), Settings->SourceTexture);
+					BackgroundMaterial->SetScalarParameterValue(TEXT("TextureSampler"), 0);
+					break;
+				case EMaterialSamplerType::SAMPLERTYPE_Grayscale:
+					BackgroundMaterial->SetTextureParameterValue(TEXT("BackgroundBaseMap_Grayscale"), Settings->SourceTexture);
+					BackgroundMaterial->SetScalarParameterValue(TEXT("TextureSampler"), 1);
+					break;
+				case EMaterialSamplerType::SAMPLERTYPE_Alpha:
+					BackgroundMaterial->SetTextureParameterValue(TEXT("BackgroundBaseMap_Alpha"), Settings->SourceTexture);
+					BackgroundMaterial->SetScalarParameterValue(TEXT("TextureSampler"), 2);
+					break;
+				case EMaterialSamplerType::SAMPLERTYPE_Normal:
+					BackgroundMaterial->SetTextureParameterValue(TEXT("BackgroundBaseMap_Normal"), Settings->SourceTexture);
+					BackgroundMaterial->SetScalarParameterValue(TEXT("TextureSampler"), 3);
+					break;
+				case EMaterialSamplerType::SAMPLERTYPE_Masks:
+					BackgroundMaterial->SetTextureParameterValue(TEXT("BackgroundBaseMap_Masks"), Settings->SourceTexture);
+					BackgroundMaterial->SetScalarParameterValue(TEXT("TextureSampler"), 4);
+					break;
+				case EMaterialSamplerType::SAMPLERTYPE_DistanceFieldFont:
+					BackgroundMaterial->SetTextureParameterValue(TEXT("BackgroundBaseMap_DistanceFeildFont"), Settings->SourceTexture);
+					BackgroundMaterial->SetScalarParameterValue(TEXT("TextureSampler"), 5);
+					break;
+				case EMaterialSamplerType::SAMPLERTYPE_LinearColor:
+					BackgroundMaterial->SetTextureParameterValue(TEXT("BackgroundBaseMap_LinearColor"), Settings->SourceTexture);
+					BackgroundMaterial->SetScalarParameterValue(TEXT("TextureSampler"), 6);
+					break;
+				case EMaterialSamplerType::SAMPLERTYPE_LinearGrayscale:
+					BackgroundMaterial->SetTextureParameterValue(TEXT("BackgroundBaseMap_LinearGrayscale"), Settings->SourceTexture);
+					BackgroundMaterial->SetScalarParameterValue(TEXT("TextureSampler"), 7);
+					break;
+				case EMaterialSamplerType::SAMPLERTYPE_Data:
+					BackgroundMaterial->SetTextureParameterValue(TEXT("BackgroundBaseMap_Data"), Settings->SourceTexture);
+					BackgroundMaterial->SetScalarParameterValue(TEXT("TextureSampler"), 8);
+					break;
+				case EMaterialSamplerType::SAMPLERTYPE_External:
+					// TODO: Enable support for External textures at a future date as needed.
+					//BackgroundMaterial->SetTextureParameterValue(TEXT("BackgroundBaseMap_Data"), Settings->SourceTexture);
+					BackgroundMaterial->SetScalarParameterValue(TEXT("TextureSampler"), 9);
+					break;
+
+				case EMaterialSamplerType::SAMPLERTYPE_VirtualColor:
+					BackgroundMaterial->SetTextureParameterValue(TEXT("BackgroundBaseMap_VTColor"), Settings->SourceTexture);
+					BackgroundMaterial->SetScalarParameterValue(TEXT("VirtualTextureSampler"), 0);
+					break;
+				case EMaterialSamplerType::SAMPLERTYPE_VirtualGrayscale:
+					BackgroundMaterial->SetTextureParameterValue(TEXT("BackgroundBaseMap_VTGrayscale"), Settings->SourceTexture);
+					BackgroundMaterial->SetScalarParameterValue(TEXT("VirtualTextureSampler"), 1);
+					break;
+				case EMaterialSamplerType::SAMPLERTYPE_VirtualAlpha:
+					BackgroundMaterial->SetTextureParameterValue(TEXT("BackgroundBaseMap_VTAlpha"), Settings->SourceTexture);
+					BackgroundMaterial->SetScalarParameterValue(TEXT("VirtualTextureSampler"), 2);
+					break;
+				case EMaterialSamplerType::SAMPLERTYPE_VirtualNormal:
+					BackgroundMaterial->SetTextureParameterValue(TEXT("BackgroundBaseMap_VTNormal"), Settings->SourceTexture);
+					BackgroundMaterial->SetScalarParameterValue(TEXT("VirtualTextureSampler"), 3);
+					break;
+				case EMaterialSamplerType::SAMPLERTYPE_VirtualMasks:
+					BackgroundMaterial->SetTextureParameterValue(TEXT("BackgroundBaseMap_VTMasks"), Settings->SourceTexture);
+					BackgroundMaterial->SetScalarParameterValue(TEXT("VirtualTextureSampler"), 4);
+					break;						
+				case EMaterialSamplerType::SAMPLERTYPE_VirtualLinearColor:
+					BackgroundMaterial->SetTextureParameterValue(TEXT("BackgroundBaseMap_VTLinearColor"), Settings->SourceTexture);
+					BackgroundMaterial->SetScalarParameterValue(TEXT("VirtualTextureSampler"), 6);
+					break;
+				case EMaterialSamplerType::SAMPLERTYPE_VirtualLinearGrayscale:
+					BackgroundMaterial->SetTextureParameterValue(TEXT("BackgroundBaseMap_VTLinearGrayscale"), Settings->SourceTexture);
+					BackgroundMaterial->SetScalarParameterValue(TEXT("VirtualTextureSampler"), 7);
+					break;
 				}
 			}
 		}
