UVEditor: Fix issue where polygroups were being incorrectly set on operation factory instances in the Unwrap Tool

diff --git a/Engine/Plugins/Experimental/UVEditor/Source/UVEditorTools/Private/UVEditorRecomputeUVsTool.cpp b/Engine/Plugins/Experimental/UVEditor/Source/UVEditorTools/Private/UVEditorRecomputeUVsTool.cpp
--- a/Engine/Plugins/Experimental/UVEditor/Source/UVEditorTools/Private/UVEditorRecomputeUVsTool.cpp
+++ b/Engine/Plugins/Experimental/UVEditor/Source/UVEditorTools/Private/UVEditorRecomputeUVsTool.cpp
@@ -54,7 +54,6 @@ void UUVEditorRecomputeUVsTool::Setup()
 	Settings->RestoreProperties(this);
 	Settings->bUDIMCVAREnabled = (FUVEditorUXSettings::CVarEnablePrototypeUDIMSupport.GetValueOnGameThread() > 0);
 	AddToolPropertySource(Settings);
-	Factories.SetNum(Targets.Num());
 
 	UContextObjectStore* ContextStore = GetToolManager()->GetContextObjectStore();
 	UVToolSelectionAPI = ContextStore->FindContext<UUVToolSelectionAPI>();
@@ -73,15 +72,14 @@ void UUVEditorRecomputeUVsTool::Setup()
 		PolygroupLayerProperties->RestoreProperties(this, TEXT("UVEditorRecomputeUVsTool"));
 		PolygroupLayerProperties->InitializeGroupLayers(Targets[0]->AppliedCanonical.Get());
 		PolygroupLayerProperties->WatchProperty(PolygroupLayerProperties->ActiveGroupLayer, [&](FName) { OnSelectedGroupLayerChanged(); });
-		AddToolPropertySource(PolygroupLayerProperties);
-		UpdateActiveGroupLayer();
+		AddToolPropertySource(PolygroupLayerProperties);		
 	}
 	else
 	{
 		Settings->bEnablePolygroupSupport = false;
 		Settings->IslandGeneration = EUVEditorRecomputeUVsPropertiesIslandMode::ExistingUVs;
 	}
-
+	UpdateActiveGroupLayer(false);  /* Don't update factories that don't exist yet. */
 
 	auto SetupOpFactory = [this](UUVEditorToolMeshInput& Target, TSet<int32>* Selection)
 	{
@@ -201,9 +199,9 @@ void UUVEditorRecomputeUVsTool::Shutdown(EToolShutdownType ShutdownType)
 	{
 		Target->AppliedPreview->ClearOpFactory();
 	}
-	for (int32 TargetIndex = 0; TargetIndex < Targets.Num(); ++TargetIndex)
+	for (int32 FactoryIndex = 0; FactoryIndex < Factories.Num(); ++FactoryIndex)
 	{
-		Factories[TargetIndex] = nullptr;
+		Factories[FactoryIndex] = nullptr;
 	}
 	Settings = nullptr;
 	Targets.Empty();
@@ -244,7 +242,7 @@ void UUVEditorRecomputeUVsTool::OnSelectedGroupLayerChanged()
 }
 
 
-void UUVEditorRecomputeUVsTool::UpdateActiveGroupLayer()
+void UUVEditorRecomputeUVsTool::UpdateActiveGroupLayer(bool bUpdateFactories)
 {
 	if (Targets.Num() == 1)
 	{
@@ -260,10 +258,16 @@ void UUVEditorRecomputeUVsTool::UpdateActiveGroupLayer()
 			ActiveGroupSet = MakeShared<UE::Geometry::FPolygroupSet, ESPMode::ThreadSafe>(Targets[0]->AppliedCanonical.Get(), FoundAttrib);
 		}
 	}
-	for (int32 TargetIndex = 0; TargetIndex < Targets.Num(); ++TargetIndex)
+	else
+	{
+		ActiveGroupSet = nullptr;
+	}
+
+	if (bUpdateFactories)
 	{
-		if (Factories[TargetIndex]) {
-			Factories[TargetIndex]->InputGroups = ActiveGroupSet;
+		for (int32 FactoryIdx = 0; FactoryIdx < Factories.Num(); ++FactoryIdx)
+		{
+			Factories[FactoryIdx]->InputGroups = ActiveGroupSet;
 		}
 	}
 }
diff --git a/Engine/Plugins/Experimental/UVEditor/Source/UVEditorTools/Public/UVEditorRecomputeUVsTool.h b/Engine/Plugins/Experimental/UVEditor/Source/UVEditorTools/Public/UVEditorRecomputeUVsTool.h
--- a/Engine/Plugins/Experimental/UVEditor/Source/UVEditorTools/Public/UVEditorRecomputeUVsTool.h
+++ b/Engine/Plugins/Experimental/UVEditor/Source/UVEditorTools/Public/UVEditorRecomputeUVsTool.h
@@ -90,7 +90,7 @@ protected:
 
 	TSharedPtr<UE::Geometry::FPolygroupSet, ESPMode::ThreadSafe> ActiveGroupSet;
 	void OnSelectedGroupLayerChanged();
-	void UpdateActiveGroupLayer();
+	void UpdateActiveGroupLayer(bool bUpdateFactories = true);
 
 	//
 	// Analytics
