UVEditor: Enable support for polygroups and UDIMS in the AutoUV tool.

diff --git a/Engine/Plugins/Experimental/UVEditor/Source/UVEditorToolsEditorOnly/Private/Operators/UVEditorParameterizeMeshOp.cpp b/Engine/Plugins/Experimental/UVEditor/Source/UVEditorToolsEditorOnly/Private/Operators/UVEditorParameterizeMeshOp.cpp
new file mode 100644
--- /dev/null
+++ b/Engine/Plugins/Experimental/UVEditor/Source/UVEditorToolsEditorOnly/Private/Operators/UVEditorParameterizeMeshOp.cpp
@@ -0,0 +1,538 @@
+// Copyright Epic Games, Inc. All Rights Reserved.
+
+#include "Operators/UVEditorParameterizeMeshOp.h"
+#include "Properties/ParameterizeMeshProperties.h"
+
+#include "DynamicMesh/DynamicMeshAttributeSet.h"
+#include "Selections/MeshConnectedComponents.h"
+
+#include "Parameterization/MeshLocalParam.h"
+#include "Parameterization/DynamicMeshUVEditor.h"
+#include "Parameterization/PatchBasedMeshUVGenerator.h"
+#include "DynamicMesh/DynamicMeshAABBTree3.h"
+#include "DynamicSubmesh3.h"
+#include "XAtlasWrapper.h"
+
+// The ProxyLOD plugin is currently only available on Windows.
+#if WITH_PROXYLOD
+#include "ProxyLODParameterization.h"
+#endif	// WITH_PROXYLOD
+
+// We do not have UVAtlas available if ProxyLOD is not available.
+#define NO_UVATLAS (WITH_PROXYLOD ? 0 : 1)
+
+using namespace UE::Geometry;
+
+#define LOCTEXT_NAMESPACE "ParameterizeMeshOp"
+
+FUVEditorParameterizeMeshOp::FLinearMesh::FLinearMesh(const FDynamicMesh3& Mesh, UE::Geometry::FPolygroupSet* InputGroups)
+{
+	TArray<FVector3f>& Positions = this->VertexBuffer;
+
+	// Temporary maps used during construction.
+
+	TArray<int32> TriFromID;
+	TArray<int32> TriToID;
+	TArray<int32> VertFromID;
+
+	// Compute the mapping from triangle ID to triangle number
+	{
+		const int32 MaxTriID = Mesh.MaxTriangleID(); // really +1.. all TriID < MaxTriID
+		const int32 NumTris = Mesh.TriangleCount();
+
+		// reserve space and add elements 
+		TriFromID.SetNum(MaxTriID);
+
+		// reserve space
+		TriToID.Empty(NumTris);
+
+		int32 count = 0;
+		for (int TriID : Mesh.TriangleIndicesItr())
+		{
+			TriToID.Add(TriID);
+			TriFromID[TriID] = count;
+			count++;
+		}
+	}
+
+	// Compute the mapping from vertex ID to vertex number
+	{
+		const int32 MaxVertID = Mesh.MaxVertexID();
+		const int32 NumVerts = Mesh.VertexCount();
+
+		// reserve space and add elements
+		VertFromID.Empty(MaxVertID);
+		VertFromID.AddUninitialized(MaxVertID);
+
+		// reserve space
+		VertToID.Empty(NumVerts);
+
+		int32 count = 0;
+		for (int VtxID : Mesh.VertexIndicesItr())
+		{
+			VertToID.Add(VtxID);
+			VertFromID[VtxID] = count;
+			count++;
+		}
+	}
+
+	// Fill the vertex buffer
+	{
+		int32 NumVerts = Mesh.VertexCount();
+		Positions.Empty(NumVerts);
+
+		for (const FVector3d& Vertex : Mesh.VerticesItr())
+		{
+			FVector3f Pos(Vertex.X, Vertex.Y, Vertex.Z);
+			Positions.Add(Pos);
+		}
+	}
+
+	const int32 NumTris = Mesh.TriangleCount();
+
+	// Fill the index buffer
+	{
+		IndexBuffer.Empty(NumTris * 3);
+		for (const FIndex3i& Tri : Mesh.TrianglesItr())
+		{
+			for (int i = 0; i < 3; ++i)
+			{
+				int VtxID = Tri[i];
+				int32 RemapVtx = VertFromID[VtxID];
+				IndexBuffer.Add(RemapVtx);
+			}
+		}
+
+	}
+
+
+	// For each edge on each triangle.
+	AdjacencyBuffer.Empty(NumTris * 3);
+
+	// Create Adjacency - create boundaries in adjacency at polygroup boundaries if requested.
+	if (InputGroups)
+	{
+		for (int TriID : Mesh.TriangleIndicesItr())
+		{
+			FIndex3i NbrTris = Mesh.GetTriNeighbourTris(TriID);
+			for (int32 i = 0; i < 3; ++i)
+			{
+				const int32 NbrID = NbrTris[i];
+
+				int32 RemapNbrID = FDynamicMesh3::InvalidID;
+				if (NbrID != FDynamicMesh3::InvalidID)
+				{
+					const int32 CurTriGroup = InputGroups->GetGroup(TriID);
+					const int32 NbrTriGroup = InputGroups->GetGroup(NbrID);
+
+					RemapNbrID = (CurTriGroup == NbrTriGroup) ? TriFromID[NbrID] : FDynamicMesh3::InvalidID;
+				}
+
+				AdjacencyBuffer.Add(RemapNbrID);
+			}
+		}
+
+	}
+	else  // compute the adjacency using only the mesh connectivity.
+	{
+		for (int TriID : Mesh.TriangleIndicesItr())
+		{
+			FIndex3i NbrTris = Mesh.GetTriNeighbourTris(TriID);
+			for (int32 i = 0; i < 3; ++i)
+			{
+				int32 NbrID = NbrTris[i];
+
+				int32 RemapNbrID = (NbrID != FDynamicMesh3::InvalidID) ? TriFromID[NbrID] : FDynamicMesh3::InvalidID;
+				AdjacencyBuffer.Add(RemapNbrID);
+			}
+		}
+	}
+}
+
+
+
+void FUVEditorParameterizeMeshOp::CopyNewUVsToMesh(
+	FDynamicMesh3& Mesh,
+	const FLinearMesh& LinearMesh,
+	const FDynamicMesh3& FlippedMesh,
+	const TArray<FVector2D>& UVVertexBuffer,
+	const TArray<int32>& UVIndexBuffer,
+	const TArray<int32>& VertexRemapArray,
+	bool bReverseOrientation)
+{
+	if (!Mesh.HasAttributes())
+	{
+		Mesh.EnableAttributes();
+	}
+	FDynamicMeshUVOverlay* UVOverlay = Mesh.Attributes()->GetUVLayer(UVLayer);
+	if (UVOverlay == nullptr)
+	{
+		Mesh.Attributes()->SetNumUVLayers(UVLayer + 1);
+		UVOverlay = Mesh.Attributes()->GetUVLayer(UVLayer);
+	}
+
+	// delete any existing UVs
+	UVOverlay->ClearElements();
+	checkSlow(UVOverlay->ElementCount() == 0);
+
+	// Add the UVs to the overlay
+	int32 NumUVs = UVVertexBuffer.Num();
+	TArray<int32> UVOffsetToElID;  UVOffsetToElID.Reserve(NumUVs);
+
+	for (int32 i = 0; i < NumUVs; ++i)
+	{
+		FVector2D UV = UVVertexBuffer[i];
+
+		// add the UV to the mesh overlay
+		const int32 NewID = UVOverlay->AppendElement(FVector2f(UV));
+		UVOffsetToElID.Add(NewID);
+	}
+
+	int32 NumUVTris = UVIndexBuffer.Num() / 3;
+	for (int32 i = 0; i < NumUVTris; ++i)
+	{
+		int32 t = i * 3;
+		// The triangle in UV space
+		FIndex3i UVTri(UVIndexBuffer[t], UVIndexBuffer[t + 1], UVIndexBuffer[t + 2]);
+
+		// the triangle in terms of the VertIDs in the DynamicMesh
+		FIndex3i TriVertIDs;
+		for (int c = 0; c < 3; ++c)
+		{
+			// the offset for this vertex in the LinearMesh
+			int32 Offset = VertexRemapArray[UVTri[c]];
+
+			int32 VertID = LinearMesh.VertToID[Offset];
+
+			TriVertIDs[c] = VertID;
+		}
+
+		// NB: this could be slow.. 
+		int32 TriID = FlippedMesh.FindTriangle(TriVertIDs[0], TriVertIDs[1], TriVertIDs[2]);
+
+		checkSlow(TriID != FDynamicMesh3::InvalidID);
+
+		FIndex3i ElTri;
+		if (bReverseOrientation)
+		{
+			ElTri = FIndex3i(UVOffsetToElID[UVTri[1]], UVOffsetToElID[UVTri[0]], UVOffsetToElID[UVTri[2]]);
+		}
+		else
+		{
+			ElTri = FIndex3i(UVOffsetToElID[UVTri[0]], UVOffsetToElID[UVTri[1]], UVOffsetToElID[UVTri[2]]);
+		}
+
+		// add the triangle to the overlay
+		UVOverlay->SetTriangle(TriID, ElTri);
+	}
+}
+
+void FUVEditorParameterizeMeshOp::LayoutToUDIMByPolygroup(FDynamicMesh3& InOutMesh, UE::Geometry::FPolygroupSet& PolygroupSet)
+{
+	FDynamicMeshUVEditor UVEditor(&InOutMesh, UVLayer, true);
+	TMap<int32, TArray<int32> > TidsPerTile;
+	
+	for (int32 Tid : InOutMesh.TriangleIndicesItr())
+	{
+		if (UVEditor.GetOverlay()->IsSetTriangle(Tid))
+		{
+			int32 TidGroup = PolygroupSet.GetGroup(Tid);
+			TidsPerTile.FindOrAdd(TidGroup).Add(Tid);
+		}
+	}
+
+	TArray<int32> Tiles;
+	TidsPerTile.GetKeys(Tiles);
+	Tiles.Sort();
+	for (int32 FlatGroupID = 0; FlatGroupID < Tiles.Num(); ++FlatGroupID)
+	{
+		FVector2i TilePos;
+		TilePos.X = FlatGroupID % 10;
+		TilePos.Y = FlatGroupID / 10;
+
+		ensure(UVEditor.UDIMPack(Width, Gutter, TilePos, TidsPerTile.Find(Tiles[FlatGroupID])));
+	}
+}
+
+
+bool FUVEditorParameterizeMeshOp::ComputeUVs_UVAtlas(FDynamicMesh3& Mesh, TFunction<bool(float)>& Interrupter)
+{
+#if NO_UVATLAS
+
+	ensureMsgf(false, TEXT("UVAtlas not available; this should not be called!"));
+	return false;
+
+#else	// NO_UVATLAS
+
+	TRACE_CPUPROFILER_EVENT_SCOPE(UVEditorParameterizeMeshOp_ComputeUVs_UVAtlas);
+
+	// the UVAtlas code is unhappy if you feed it a single degenerate triangle
+	bool bNonDegenerate = true;
+	if (Mesh.TriangleCount() == 1)
+	{
+		for (int32 TriID : Mesh.TriangleIndicesItr())
+		{
+			double Area = Mesh.GetTriArea(TriID);
+			bNonDegenerate = bNonDegenerate && FMath::Abs(Area) > 1.e-5;
+		}
+	}
+
+	if (!bNonDegenerate)
+	{
+		NewResultInfo.AddError(FGeometryError(0, LOCTEXT("DegenerateMeshError", "Mesh Contains Degenerate Triangles, Cannot be processed by UVAtlas")));
+		return false;
+	}
+
+	// IProxyLODParameterization module calls into UVAtlas, which (appears to) assume mesh orientation that is opposite
+	// if what we use in UE. As a result the UV islands come back mirrored. So we send a reverse-orientation mesh instead,
+	// and then fix up the UVs when we set them below. 
+	const bool bFixOrientation = true;
+	FDynamicMesh3 FlippedMesh(EMeshComponents::FaceGroups);
+	FlippedMesh.Copy(Mesh, false, false, false, false);
+	if (bFixOrientation)
+	{
+		FlippedMesh.ReverseOrientation(false);
+	}
+
+	// Convert to a dense form.
+	FLinearMesh LinearMesh(FlippedMesh, (bRespectInputGroups && InputGroups.IsValid()) ? InputGroups.Get() : nullptr );
+
+	// Data to be populated by the UV generation tool
+	TArray<FVector2D> UVVertexBuffer;
+	TArray<int32>     UVIndexBuffer;
+	TArray<int32>     VertexRemapArray; // This maps the UV vertices to the original position vertices.  Note multiple UV vertices might share the same positional vertex (due to UV boundaries)
+
+
+	float MaxStretch = Stretch;
+	int32 MaxChartNumber = NumCharts;
+
+	TUniquePtr<IProxyLODParameterization> ParameterizationTool = IProxyLODParameterization::CreateTool();
+	bool bSuccess = ParameterizationTool->GenerateUVs(Width, Height, Gutter, LinearMesh.VertexBuffer,
+		LinearMesh.IndexBuffer, LinearMesh.AdjacencyBuffer, Interrupter,
+		UVVertexBuffer, UVIndexBuffer, VertexRemapArray, MaxStretch,
+		MaxChartNumber);
+
+	// Add the UVs to the FDynamicMesh
+	if (bSuccess)
+	{
+		CopyNewUVsToMesh(Mesh, LinearMesh, FlippedMesh, UVVertexBuffer, UVIndexBuffer, VertexRemapArray, bFixOrientation);
+	}
+	else
+	{
+		NewResultInfo.AddError(FGeometryError(0, LOCTEXT("UVAtlasFailed", "UVAtlas failed")));
+	}
+
+	if (bPackToUDIMSByOriginPolygroup && bRespectInputGroups && InputGroups.IsValid())
+	{
+		LayoutToUDIMByPolygroup(Mesh, *InputGroups);
+	}
+
+	return bSuccess;
+
+#endif	// NO_UVATLAS
+}
+
+
+
+
+
+
+
+
+
+bool FUVEditorParameterizeMeshOp::ComputeUVs_XAtlas(FDynamicMesh3& Mesh, TFunction<bool(float)>& Interrupter)
+{
+	TRACE_CPUPROFILER_EVENT_SCOPE(UVEditorParameterizeMeshOp_ComputeUVs_XAtlas);
+
+	// reverse mesh orientation
+	const bool bFixOrientation = true;
+	FDynamicMesh3 FlippedMesh(EMeshComponents::FaceGroups);
+	FlippedMesh.Copy(Mesh, false, false, false, false);
+	if (bFixOrientation)
+	{
+		FlippedMesh.ReverseOrientation(false);
+	}
+
+	// Convert to a dense form.
+	FLinearMesh LinearMesh(FlippedMesh, (bRespectInputGroups && InputGroups.IsValid()) ? InputGroups.Get() : nullptr);
+
+	// Data to be populated by the UV generation tool
+	TArray<FVector2D> UVVertexBuffer;
+	TArray<int32>     UVIndexBuffer;
+	TArray<int32>     VertexRemapArray; // This maps the UV vertices to the original position vertices.  Note multiple UV vertices might share the same positional vertex (due to UV boundaries)
+
+	XAtlasWrapper::XAtlasChartOptions ChartOptions;
+	ChartOptions.MaxIterations = XAtlasMaxIterations;
+	XAtlasWrapper::XAtlasPackOptions PackOptions;
+	bool bSuccess = XAtlasWrapper::ComputeUVs(LinearMesh.IndexBuffer,
+		LinearMesh.VertexBuffer,
+		ChartOptions,
+		PackOptions,
+		UVVertexBuffer,
+		UVIndexBuffer,
+		VertexRemapArray);
+
+	// Add the UVs to the FDynamicMesh
+	if (bSuccess)
+	{
+		CopyNewUVsToMesh(Mesh, LinearMesh, FlippedMesh, UVVertexBuffer, UVIndexBuffer, VertexRemapArray, bFixOrientation);
+	}
+	else
+	{
+		NewResultInfo.AddError(FGeometryError(0, LOCTEXT("XAtlasFailed", "XAtlas failed")));
+	}
+
+	return bSuccess;
+}
+
+
+
+
+bool FUVEditorParameterizeMeshOp::ComputeUVs_PatchBuilder(FDynamicMesh3& InOutMesh, FProgressCancel* ProgressCancel)
+{
+	TRACE_CPUPROFILER_EVENT_SCOPE(UVEditorParameterizeMeshOp_ComputeUVs_PatchBuilder);
+
+	FPatchBasedMeshUVGenerator UVGenerator;
+
+	TUniquePtr<FPolygroupSet> PolygroupConstraint;
+	if (bRespectInputGroups && InputGroups.IsValid())
+	{
+		UVGenerator.GroupConstraint = InputGroups.Get();
+	}
+
+	UVGenerator.TargetPatchCount = FMath::Max(1, InitialPatchCount);
+	UVGenerator.bNormalWeightedPatches = true;
+	UVGenerator.PatchNormalWeight = FMath::Clamp(PatchCurvatureAlignmentWeight, 0.0, 999999.0);
+	UVGenerator.MinPatchSize = 2;
+
+	UVGenerator.MergingThreshold = FMath::Clamp(PatchMergingMetricThresh, 0.001, 9999.0);
+	UVGenerator.MaxNormalDeviationDeg = FMath::Clamp(PatchMergingAngleThresh, 0.0, 180.0);
+
+	UVGenerator.NormalSmoothingRounds = FMath::Clamp(ExpMapNormalSmoothingSteps, 0, 9999);
+	UVGenerator.NormalSmoothingAlpha = FMath::Clamp(ExpMapNormalSmoothingAlpha, 0.0, 1.0);
+
+	UVGenerator.bAutoAlignPatches = true;
+	UVGenerator.bAutoPack = bEnablePacking;
+	UVGenerator.PackingTextureResolution = FMath::Clamp(this->Width, 16, 4096);
+	UVGenerator.PackingGutterWidth = this->Gutter;
+
+
+	FDynamicMeshUVEditor UVEditor(&InOutMesh, UVLayer, true);
+	FGeometryResult Result = UVGenerator.AutoComputeUVs(*UVEditor.GetMesh(), *UVEditor.GetOverlay(), ProgressCancel);
+
+	if (bPackToUDIMSByOriginPolygroup && bRespectInputGroups && InputGroups.IsValid())
+	{
+		LayoutToUDIMByPolygroup(InOutMesh, *InputGroups);
+	}
+
+	SetResultInfo(NewResultInfo);
+	return (Result.HasFailed() == false);
+}
+
+
+
+void FUVEditorParameterizeMeshOp::CalculateResult(FProgressCancel* Progress)
+{
+	TRACE_CPUPROFILER_EVENT_SCOPE(UVEditorParameterizeMeshOp_CalculateResult);
+
+	if (!InputMesh.IsValid())
+	{
+		SetResultInfo(FGeometryResult::Failed());
+		return;
+	}
+
+	ResultMesh = MakeUnique<FDynamicMesh3>(*InputMesh);
+	if (Progress && Progress->Cancelled())
+	{
+		SetResultInfo(FGeometryResult::Cancelled());
+		return;
+	}
+
+	// if using PatchBuilder, it handles all the output stuff
+	if (Method == EUVEditorParamOpBackend::PatchBuilder)
+	{
+		ComputeUVs_PatchBuilder(*ResultMesh, Progress);
+		return;
+	}
+
+	// The UV atlas callback uses a float progress-based interrupter. 
+	TFunction<bool(float)> Interrupter = [Progress](float)->bool {return !(Progress && Progress->Cancelled()); };
+
+#if NO_UVATLAS
+	constexpr bool bUseXAtlas = true;
+	auto LogUVAtlasWarning = []
+	{
+		UE_LOG(LogGeometry, Warning, TEXT("AutoUV method UVAtlas not available; falling back to XAtlas instead."));
+	};
+	UE_CALL_ONCE(LogUVAtlasWarning);
+#else	// NO_UVATLAS
+	const bool bUseXAtlas = (Method == EUVEditorParamOpBackend::XAtlas);
+#endif	// NO_UVATLAS
+
+	NewResultInfo = FGeometryResult(EGeometryResultType::InProgress);
+
+	// Single call to UVAtlas - perhaps respecting the poly groups.
+	bool bOK;
+	if (bUseXAtlas)
+	{
+		bOK = ComputeUVs_XAtlas(*ResultMesh, Interrupter);
+	}
+	else
+	{
+		bOK = ComputeUVs_UVAtlas(*ResultMesh, Interrupter);
+	}
+
+	NewResultInfo.SetSuccess(bOK, Progress);
+	SetResultInfo(NewResultInfo);
+}
+
+
+
+TUniquePtr<FDynamicMeshOperator> UUVEditorParameterizeMeshOperatorFactory::MakeNewOperator()
+{
+	TUniquePtr<FUVEditorParameterizeMeshOp> Op = MakeUnique<FUVEditorParameterizeMeshOp>();
+
+	Op->InputMesh = OriginalMesh;
+	Op->InputGroups = InputGroups;
+
+	switch (Settings->Method)
+	{
+	case EUVEditorParameterizeMeshUVMethod::PatchBuilder:
+		Op->Method = EUVEditorParamOpBackend::PatchBuilder;
+		Op->InitialPatchCount = PatchBuilderProperties->InitialPatches;
+		Op->bRespectInputGroups = PatchBuilderProperties->bUsePolygroups;
+		Op->PatchCurvatureAlignmentWeight = PatchBuilderProperties->CurvatureAlignment;
+		Op->PatchMergingMetricThresh = PatchBuilderProperties->MergingDistortionThreshold;
+		Op->PatchMergingAngleThresh = PatchBuilderProperties->MergingAngleThreshold;
+		Op->bPackToUDIMSByOriginPolygroup = PatchBuilderProperties->bLayoutUDIMPerPolygroup;
+		Op->ExpMapNormalSmoothingSteps = PatchBuilderProperties->SmoothingSteps;
+		Op->ExpMapNormalSmoothingAlpha = PatchBuilderProperties->SmoothingAlpha;
+		Op->bEnablePacking = PatchBuilderProperties->bRepack;
+		Op->Width = PatchBuilderProperties->TextureResolution;
+		Op->Height = PatchBuilderProperties->TextureResolution;
+		break;
+	case EUVEditorParameterizeMeshUVMethod::UVAtlas:
+		Op->Method = EUVEditorParamOpBackend::UVAtlas;
+		Op->Stretch = UVAtlasProperties->IslandStretch;
+		Op->bRespectInputGroups = UVAtlasProperties->bUsePolygroups;
+		Op->bPackToUDIMSByOriginPolygroup = UVAtlasProperties->bLayoutUDIMPerPolygroup;
+		Op->NumCharts = UVAtlasProperties->NumIslands;
+		Op->Width = UVAtlasProperties->TextureResolution;
+		Op->Height = UVAtlasProperties->TextureResolution;
+		break;
+	case EUVEditorParameterizeMeshUVMethod::XAtlas:
+		Op->Method = EUVEditorParamOpBackend::XAtlas;
+		Op->XAtlasMaxIterations = XAtlasProperties->MaxIterations;
+		break;
+	}
+
+	Op->UVLayer = GetSelectedUVChannel();
+
+	FTransformSRT3d Transform(TargetTransform);
+	Op->SetTransform(Transform);
+
+	return Op;
+}
+
+
+#undef LOCTEXT_NAMESPACE
\ No newline at end of file
diff --git a/Engine/Plugins/Experimental/UVEditor/Source/UVEditorToolsEditorOnly/Private/UVEditorParameterizeMeshTool.cpp b/Engine/Plugins/Experimental/UVEditor/Source/UVEditorToolsEditorOnly/Private/UVEditorParameterizeMeshTool.cpp
--- a/Engine/Plugins/Experimental/UVEditor/Source/UVEditorToolsEditorOnly/Private/UVEditorParameterizeMeshTool.cpp
+++ b/Engine/Plugins/Experimental/UVEditor/Source/UVEditorToolsEditorOnly/Private/UVEditorParameterizeMeshTool.cpp
@@ -4,12 +4,14 @@
 #include "InteractiveToolManager.h"
 #include "ToolBuilderUtil.h"
 #include "DynamicMesh/DynamicMesh3.h"
-#include "ParameterizationOps/ParameterizeMeshOp.h"
-#include "Properties/ParameterizeMeshProperties.h"
+#include "Operators/UVEditorParameterizeMeshOp.h"
+#include "PropertySets/PolygroupLayersProperties.h"
+#include "Polygroups/PolygroupUtil.h"
 #include "ToolTargets/UVEditorToolMeshInput.h"
 #include "ContextObjects/UVToolContextObjects.h"
 #include "ContextObjectStore.h"
 #include "EngineAnalytics.h"
+#include "UVEditorUXSettings.h"
 
 using namespace UE::Geometry;
 
@@ -49,33 +51,56 @@ void UUVEditorParameterizeMeshTool::Setup()
 	UInteractiveTool::Setup();
 
 	// initialize our properties
-	Settings = NewObject<UParameterizeMeshToolProperties>(this);
+	Settings = NewObject<UUVEditorParameterizeMeshToolProperties>(this);
 	Settings->RestoreProperties(this);
 	AddToolPropertySource(Settings);
-	Settings->WatchProperty(Settings->Method, [&](EParameterizeMeshUVMethod) { OnMethodTypeChanged(); });
+	Settings->WatchProperty(Settings->Method, [&](EUVEditorParameterizeMeshUVMethod) { OnMethodTypeChanged(); });
 
-	UVAtlasProperties = NewObject<UParameterizeMeshToolUVAtlasProperties>(this);
+	UVAtlasProperties = NewObject<UUVEditorParameterizeMeshToolUVAtlasProperties>(this);
 	UVAtlasProperties->RestoreProperties(this);
 	AddToolPropertySource(UVAtlasProperties);
 	SetToolPropertySourceEnabled(UVAtlasProperties, true);
 
-	XAtlasProperties = NewObject<UParameterizeMeshToolXAtlasProperties>(this);
+	XAtlasProperties = NewObject<UUVEditorParameterizeMeshToolXAtlasProperties>(this);
 	XAtlasProperties->RestoreProperties(this);
 	AddToolPropertySource(XAtlasProperties);
 	SetToolPropertySourceEnabled(XAtlasProperties, true);
 	
-	PatchBuilderProperties = NewObject<UParameterizeMeshToolPatchBuilderProperties>(this);
+	PatchBuilderProperties = NewObject<UUVEditorParameterizeMeshToolPatchBuilderProperties>(this);
 	PatchBuilderProperties->RestoreProperties(this);
 	AddToolPropertySource(PatchBuilderProperties);
 	SetToolPropertySourceEnabled(PatchBuilderProperties, true);
 
+	if (Targets.Num() == 1)
+	{
+		PolygroupLayerProperties = NewObject<UPolygroupLayersProperties>(this);
+		PolygroupLayerProperties->RestoreProperties(this, TEXT("UVEditorRecomputeUVsTool"));
+		PolygroupLayerProperties->InitializeGroupLayers(Targets[0]->AppliedCanonical.Get());
+		PolygroupLayerProperties->WatchProperty(PolygroupLayerProperties->ActiveGroupLayer, [&](FName) { OnSelectedGroupLayerChanged(); });
+		AddToolPropertySource(PolygroupLayerProperties);
+		PatchBuilderProperties->bPolygroupsEnabled = true;
+		UVAtlasProperties->bPolygroupsEnabled = true;
+	}
+	else
+	{
+		ActiveGroupSet = nullptr;
+		PatchBuilderProperties->bPolygroupsEnabled = false;
+		UVAtlasProperties->bPolygroupsEnabled = false;
+	}
+	UpdateActiveGroupLayer(false);  /* Don't update factories that don't exist yet. */
+
+	PatchBuilderProperties->bUDIMsEnabled = (FUVEditorUXSettings::CVarEnablePrototypeUDIMSupport.GetValueOnGameThread() > 0);
+	UVAtlasProperties->bUDIMsEnabled = (FUVEditorUXSettings::CVarEnablePrototypeUDIMSupport.GetValueOnGameThread() > 0);
+
+
 	Factories.SetNum(Targets.Num());
 	for (int32 TargetIndex = 0; TargetIndex < Targets.Num(); ++TargetIndex)
 	{
 		TObjectPtr<UUVEditorToolMeshInput> Target = Targets[TargetIndex];
-		Factories[TargetIndex] = NewObject<UParameterizeMeshOperatorFactory>();
+		Factories[TargetIndex] = NewObject<UUVEditorParameterizeMeshOperatorFactory>();
 		Factories[TargetIndex]->TargetTransform = Target->AppliedPreview->PreviewMesh->GetTransform();
 		Factories[TargetIndex]->Settings = Settings;
+		Factories[TargetIndex]->InputGroups = ActiveGroupSet;
 		Factories[TargetIndex]->UVAtlasProperties = UVAtlasProperties;
 		Factories[TargetIndex]->XAtlasProperties = XAtlasProperties;
 		Factories[TargetIndex]->PatchBuilderProperties = PatchBuilderProperties;
@@ -115,9 +140,12 @@ void UUVEditorParameterizeMeshTool::OnMethodTypeChanged()
 {
 	TRACE_CPUPROFILER_EVENT_SCOPE(UVEditorParameterizeMeshTool_OnMethodTypeChanged);
 	
-	SetToolPropertySourceEnabled(UVAtlasProperties, Settings->Method == EParameterizeMeshUVMethod::UVAtlas);
-	SetToolPropertySourceEnabled(XAtlasProperties, Settings->Method == EParameterizeMeshUVMethod::XAtlas);
-	SetToolPropertySourceEnabled(PatchBuilderProperties, Settings->Method == EParameterizeMeshUVMethod::PatchBuilder);
+	SetToolPropertySourceEnabled(UVAtlasProperties, Settings->Method == EUVEditorParameterizeMeshUVMethod::UVAtlas);
+	SetToolPropertySourceEnabled(XAtlasProperties, Settings->Method == EUVEditorParameterizeMeshUVMethod::XAtlas);
+	SetToolPropertySourceEnabled(PatchBuilderProperties, Settings->Method == EUVEditorParameterizeMeshUVMethod::PatchBuilder);
+
+	SetToolPropertySourceEnabled(PolygroupLayerProperties, Settings->Method == EUVEditorParameterizeMeshUVMethod::UVAtlas ||
+		                                                   Settings->Method == EUVEditorParameterizeMeshUVMethod::PatchBuilder);
 
 	for (TObjectPtr<UUVEditorToolMeshInput> Target : Targets)
 	{
@@ -182,15 +210,55 @@ void UUVEditorParameterizeMeshTool::Shutdown(EToolShutdownType ShutdownType)
 		Target->AppliedPreview->OverrideMaterial = nullptr;
 	}
 
-	for (int32 TargetIndex = 0; TargetIndex < Targets.Num(); ++TargetIndex)
+	for (int32 FactoryIndex = 0; FactoryIndex < Factories.Num(); ++FactoryIndex)
 	{
-		Factories[TargetIndex] = nullptr;
+		Factories[FactoryIndex] = nullptr;
 	}
 
 	Settings = nullptr;
 	Targets.Empty();
 }
 
+void UUVEditorParameterizeMeshTool::OnSelectedGroupLayerChanged()
+{
+	UpdateActiveGroupLayer();
+	for (TObjectPtr<UUVEditorToolMeshInput> Target : Targets)
+	{
+		Target->AppliedPreview->InvalidateResult();
+	}
+}
+
+
+void UUVEditorParameterizeMeshTool::UpdateActiveGroupLayer(bool bUpdateFactories)
+{
+	if (Targets.Num() == 1)
+	{
+		if (PolygroupLayerProperties->HasSelectedPolygroup() == false)
+		{
+			ActiveGroupSet = MakeShared<UE::Geometry::FPolygroupSet, ESPMode::ThreadSafe>(Targets[0]->AppliedCanonical.Get());
+		}
+		else
+		{
+			FName SelectedName = PolygroupLayerProperties->ActiveGroupLayer;
+			FDynamicMeshPolygroupAttribute* FoundAttrib = UE::Geometry::FindPolygroupLayerByName(*Targets[0]->AppliedCanonical, SelectedName);
+			ensureMsgf(FoundAttrib, TEXT("Selected attribute not found! Falling back to Default group layer."));
+			ActiveGroupSet = MakeShared<UE::Geometry::FPolygroupSet, ESPMode::ThreadSafe>(Targets[0]->AppliedCanonical.Get(), FoundAttrib);
+		}
+	}
+	else
+	{
+		ActiveGroupSet = nullptr;
+	}
+
+	if (bUpdateFactories)
+	{
+		for (int32 FactoryIdx = 0; FactoryIdx < Factories.Num(); ++FactoryIdx)
+		{
+			Factories[FactoryIdx]->InputGroups = ActiveGroupSet;
+		}
+	}
+}
+
 void UUVEditorParameterizeMeshTool::OnTick(float DeltaTime)
 {
 	for (TObjectPtr<UUVEditorToolMeshInput> Target : Targets)
@@ -244,11 +312,11 @@ void UUVEditorParameterizeMeshTool::RecordAnalytics()
 	Attributes.Add(FAnalyticsEventAttribute(TEXT("Stats.ToolActiveDuration"), (FDateTime::UtcNow() - ToolStartTimeAnalytics).ToString()));
 
 	// Tool settings chosen by the user (Volatile! Sync with EditCondition meta-tags in *Properties members)
-	const FString MethodName = StaticEnum<EParameterizeMeshUVMethod>()->GetNameStringByIndex(static_cast<int>(Settings->Method));
+	const FString MethodName = StaticEnum<EUVEditorParameterizeMeshUVMethod>()->GetNameStringByIndex(static_cast<int>(Settings->Method));
 	Attributes.Add(FAnalyticsEventAttribute(TEXT("Settings.Method"), MethodName));
 	switch (Settings->Method)
 	{
-		case EParameterizeMeshUVMethod::PatchBuilder:
+		case EUVEditorParameterizeMeshUVMethod::PatchBuilder:
 			Attributes.Add(FAnalyticsEventAttribute(FString::Printf(TEXT("Settings.%s.InitialPatches"), *MethodName), PatchBuilderProperties->InitialPatches));
 			Attributes.Add(FAnalyticsEventAttribute(FString::Printf(TEXT("Settings.%s.CurvatureAlignment"), *MethodName), PatchBuilderProperties->CurvatureAlignment));
 			Attributes.Add(FAnalyticsEventAttribute(FString::Printf(TEXT("Settings.%s.MergingDistortionThreshold"), *MethodName), PatchBuilderProperties->MergingDistortionThreshold));
@@ -261,11 +329,11 @@ void UUVEditorParameterizeMeshTool::RecordAnalytics()
 				Attributes.Add(FAnalyticsEventAttribute(FString::Printf(TEXT("Settings.%s.TextureResolution"), *MethodName), PatchBuilderProperties->TextureResolution));
 			}
 			break;
-		case EParameterizeMeshUVMethod::UVAtlas:
+		case EUVEditorParameterizeMeshUVMethod::UVAtlas:
 			Attributes.Add(FAnalyticsEventAttribute(FString::Printf(TEXT("Settings.%s.IslandStretch"), *MethodName), UVAtlasProperties->IslandStretch));
 			Attributes.Add(FAnalyticsEventAttribute(FString::Printf(TEXT("Settings.%s.NumIslands"), *MethodName), UVAtlasProperties->NumIslands));
 			break;
-		case EParameterizeMeshUVMethod::XAtlas:
+		case EUVEditorParameterizeMeshUVMethod::XAtlas:
 			Attributes.Add(FAnalyticsEventAttribute(FString::Printf(TEXT("Settings.%s.MaxIterations"), *MethodName), XAtlasProperties->MaxIterations));
 			break;
 		default:
diff --git a/Engine/Plugins/Experimental/UVEditor/Source/UVEditorToolsEditorOnly/Public/Operators/UVEditorParameterizeMeshOp.h b/Engine/Plugins/Experimental/UVEditor/Source/UVEditorToolsEditorOnly/Public/Operators/UVEditorParameterizeMeshOp.h
new file mode 100644
--- /dev/null
+++ b/Engine/Plugins/Experimental/UVEditor/Source/UVEditorToolsEditorOnly/Public/Operators/UVEditorParameterizeMeshOp.h
@@ -0,0 +1,285 @@
+// Copyright Epic Games, Inc. All Rights Reserved.
+
+#pragma once
+
+#include "CoreMinimal.h"
+#include "Util/ProgressCancel.h"
+#include "ModelingOperators.h"
+#include "Polygroups/PolygroupSet.h"
+#include "DynamicMesh/DynamicMeshAttributeSet.h"
+#include "InteractiveTool.h"
+
+#include "UVEditorParameterizeMeshOp.generated.h"
+
+UENUM()
+enum class EUVEditorParameterizeMeshUVMethod
+{
+	// Keep values the same as UE::Geometry::EUVEditorParamOpBackend!
+
+	/** Compute automatic UVs using the Patch Builder technique */
+	PatchBuilder = 0,
+	/** Compute automatic UVs using the UVAtlas technique */
+	UVAtlas = 1,
+	/** Compute automatic UVs using the XAtlas technique */
+	XAtlas = 2,
+};
+
+
+UCLASS()
+class UVEDITORTOOLSEDITORONLY_API UUVEditorParameterizeMeshToolProperties : public UInteractiveToolPropertySet
+{
+	GENERATED_BODY()
+public:
+	/** Automatic UV generation technique to use */
+	UPROPERTY(EditAnywhere, Category = AutoUV)
+	EUVEditorParameterizeMeshUVMethod Method = EUVEditorParameterizeMeshUVMethod::PatchBuilder;
+};
+
+
+/**
+ * Settings for the UVAtlas Automatic UV Generation Method
+ */
+UCLASS()
+class UVEDITORTOOLSEDITORONLY_API UUVEditorParameterizeMeshToolUVAtlasProperties : public UInteractiveToolPropertySet
+{
+	GENERATED_BODY()
+public:
+	/** Maximum amount of stretch, from none to unbounded. If zero stretch is specified, each triangle will likely be its own UV island. */
+	UPROPERTY(EditAnywhere, Category = UVAtlas, meta = (UIMin = "0", UIMax = "1", ClampMin = "0", ClampMax = "1"))
+	float IslandStretch = 0.11f;
+
+	/** Hint at number of UV islands. The default of 0 means it is determined automatically. */
+	UPROPERTY(EditAnywhere, Category = UVAtlas, meta = (UIMin = "0", UIMax = "100", ClampMin = "0", ClampMax = "10000"))
+	int NumIslands = 0;
+
+	/** Expected resolution of the output textures; this controls spacing left between UV islands to avoid interpolation artifacts. */
+	UPROPERTY(EditAnywhere, Category = UVAtlas, meta = (UIMin = "64", UIMax = "2048", ClampMin = "2", ClampMax = "4096"))
+	int TextureResolution = 1024;
+
+	/** Generate new UVs based on polygroups from specified layer. */
+	UPROPERTY(EditAnywhere, Category = UVAtlas, meta = (EditCondition = bPolygroupsEnabled, EditConditionHides, HideEditConditionToggle))
+	bool bUsePolygroups;
+
+	/** Layout resulting islands on UDIMs based on polygroups. */
+	UPROPERTY(EditAnywhere, Category = UVAtlas, meta = (EditCondition = "bUsePolygroups && bUDIMsEnabled", EditConditionHides, HideEditConditionToggle))
+	bool bLayoutUDIMPerPolygroup;
+
+	/** Controls if polygroup options are available to the user. */
+	UPROPERTY()
+	bool bPolygroupsEnabled;
+
+	/** Controls if UDIM options are available to the user. */
+	UPROPERTY()
+	bool bUDIMsEnabled;
+};
+
+
+UCLASS()
+class UVEDITORTOOLSEDITORONLY_API UUVEditorParameterizeMeshToolXAtlasProperties : public UInteractiveToolPropertySet
+{
+	GENERATED_BODY()
+public:
+	/** Number of solve iterations; higher values generally result in better UV islands. */
+	UPROPERTY(EditAnywhere, Category = XAtlas, meta = (UIMin = "1", UIMax = "10", ClampMin = "1", ClampMax = "1000"))
+	int MaxIterations = 1;
+};
+
+
+UCLASS()
+class UVEDITORTOOLSEDITORONLY_API UUVEditorParameterizeMeshToolPatchBuilderProperties : public UInteractiveToolPropertySet
+{
+	GENERATED_BODY()
+public:
+	/** Number of initial patches the mesh will be split into before island merging. */
+	UPROPERTY(EditAnywhere, Category = PatchBuilder, meta = (UIMin = "1", UIMax = "1000", ClampMin = "1", ClampMax = "99999999"))
+	int InitialPatches = 100;
+
+	/** Alignment of the initial patches to creases in the mesh.*/
+	UPROPERTY(EditAnywhere, Category = PatchBuilder, meta = (UIMin = "0.1", UIMax = "2.0", ClampMin = "0.01", ClampMax = "100.0"))
+	float CurvatureAlignment = 1.0f;
+
+	/** Threshold for stretching and distortion below which island merging is allowed; larger values increase the allowable UV distortion. */
+	UPROPERTY(EditAnywhere, Category = PatchBuilder, meta = (DisplayName = "Distortion Threshold", UIMin = "1.0", UIMax = "5.0", ClampMin = "1.0"))
+	float MergingDistortionThreshold = 1.5f;
+
+	/** Threshold for the average face normal deviation below which island merging is allowed. */
+	UPROPERTY(EditAnywhere, Category = PatchBuilder, meta = (DisplayName = "Angle Threshold", UIMin = "0.0", UIMax = "90.0", ClampMin = "0.0", ClampMax = "180.0"))
+	float MergingAngleThreshold = 45.0f;
+
+	/** Number of smoothing steps to apply; this slightly increases distortion but produces more stable results. */
+	UPROPERTY(EditAnywhere, Category = PatchBuilder, meta = (UIMin = "0", UIMax = "25", ClampMin = "0", ClampMax = "1000"))
+	int SmoothingSteps = 5;
+
+	/** Smoothing parameter; larger values result in faster smoothing in each step. */
+	UPROPERTY(EditAnywhere, Category = PatchBuilder, meta = (UIMin = "0", UIMax = "1.0", ClampMin = "0", ClampMax = "1.0"))
+	float SmoothingAlpha = 0.25f;
+
+	/** Automatically pack result UVs into the unit square, i.e. fit between 0 and 1 with no overlap. */
+	UPROPERTY(EditAnywhere, Category = PatchBuilder)
+	bool bRepack = true;
+
+	/** Expected resolution of the output textures; this controls spacing left between UV islands to avoid interpolation artifacts. This is only enabled when Repack is enabled. */
+	UPROPERTY(EditAnywhere, Category = PatchBuilder, meta = (UIMin = "64", UIMax = "2048", ClampMin = "2", ClampMax = "4096", EditCondition = bRepack))
+	int TextureResolution = 1024;
+
+	/** Generate new UVs based on polygroups from specified layer. */
+	UPROPERTY(EditAnywhere, Category = PatchBuilder, meta = (EditCondition = bPolygroupsEnabled, EditConditionHides, HideEditConditionToggle))
+	bool bUsePolygroups;
+
+	/** Layout resulting islands on UDIMs based on polygroups. */
+	UPROPERTY(EditAnywhere, Category = PatchBuilder, meta = (EditCondition = "bUsePolygroups && bUDIMsEnabled", EditConditionHides, HideEditConditionToggle))
+	bool bLayoutUDIMPerPolygroup;
+
+	/** Controls if polygroup options are available to the user. */
+	UPROPERTY()
+	bool bPolygroupsEnabled;
+
+	/** Controls if UDIM options are available to the user. */
+	UPROPERTY()
+	bool bUDIMsEnabled;
+};
+
+
+namespace UE
+{
+	namespace Geometry
+	{
+
+		enum class EUVEditorParamOpBackend
+		{
+			PatchBuilder = 0,
+			UVAtlas = 1,
+			XAtlas = 2
+		};
+
+		class UVEDITORTOOLSEDITORONLY_API FUVEditorParameterizeMeshOp : public FDynamicMeshOperator
+		{
+		public:
+			virtual ~FUVEditorParameterizeMeshOp() {}
+
+			//
+			// Inputs
+			// 
+
+			// source mesh
+			TSharedPtr<FDynamicMesh3, ESPMode::ThreadSafe> InputMesh;
+
+			// source groups (optional)
+			TSharedPtr<UE::Geometry::FPolygroupSet, ESPMode::ThreadSafe> InputGroups;
+
+			// UVAtlas generation parameters
+			float Stretch = 0.11f;
+			int32 NumCharts = 0;
+
+			// XAtlas generation parameters
+			int32 XAtlasMaxIterations = 1;
+
+			// PatchBuilder generation parameters
+			int32 InitialPatchCount = 100;
+			bool bRespectInputGroups = false;
+			float PatchCurvatureAlignmentWeight = 1.0;
+			float PatchMergingMetricThresh = 1.5;
+			float PatchMergingAngleThresh = 45.0;
+			int ExpMapNormalSmoothingSteps = 5;
+			float ExpMapNormalSmoothingAlpha = 0.25;
+
+			// UV layer
+			int32 UVLayer = 0;
+
+			// UDIM support
+			bool bPackToUDIMSByOriginPolygroup = false;
+
+			// Atlas Packing parameters
+			bool bEnablePacking = true;
+			int32 Height = 512;
+			int32 Width = 512;
+			float Gutter = 2.5;
+
+			EUVEditorParamOpBackend Method = EUVEditorParamOpBackend::UVAtlas;
+
+			// set ability on protected transform.
+			void SetTransform(const FTransformSRT3d& XForm)
+			{
+				ResultTransform = XForm;
+			}
+
+			//
+			// FDynamicMeshOperator implementation
+			// 
+
+			virtual void CalculateResult(FProgressCancel* Progress) override;
+
+
+		protected:
+
+			// dense index/vertex buffer based representation of the data needed for parameterization
+			struct FLinearMesh
+			{
+				FLinearMesh(const FDynamicMesh3& Mesh, UE::Geometry::FPolygroupSet* InputGroups);
+
+				// Stripped down mesh
+				TArray<int32>   IndexBuffer;
+				TArray<FVector3f> VertexBuffer;
+
+				// Map from offset in the VertexBuffer to the VertexID in the FDynamicMesh
+				TArray<int32> VertToID;
+
+				// Adjacency[tri], Adjacency[tri+1], Adjacency[tri+2] 
+				// are the three tris adjacent to tri.
+				TArray<int32>   AdjacencyBuffer;
+
+			};
+
+			FGeometryResult NewResultInfo;
+
+			bool ComputeUVs_UVAtlas(FDynamicMesh3& InOutMesh, TFunction<bool(float)>& Interrupter);
+			bool ComputeUVs_XAtlas(FDynamicMesh3& InOutMesh, TFunction<bool(float)>& Interrupter);
+			bool ComputeUVs_PatchBuilder(FDynamicMesh3& InOutMesh, FProgressCancel* Progress);
+
+			void CopyNewUVsToMesh(
+				FDynamicMesh3& Mesh,
+				const FLinearMesh& LinearMesh,
+				const FDynamicMesh3& FlippedMesh,
+				const TArray<FVector2D>& UVVertexBuffer,
+				const TArray<int32>& UVIndexBuffer,
+				const TArray<int32>& VertexRemapArray,
+				bool bReverseOrientation);
+
+			void LayoutToUDIMByPolygroup(FDynamicMesh3& InOutMesh, UE::Geometry::FPolygroupSet& InputGroups);
+		};
+
+	} // end namespace UE::Geometry
+} // end namespace UE
+
+/**
+ * Can be hooked up to a UMeshOpPreviewWithBackgroundCompute to perform UV parameterization operations.
+ *
+ * Inherits from UObject so that it can hold a strong pointer to the settings UObject, which
+ * needs to be a UObject to be displayed in the details panel.
+ */
+UCLASS()
+class UVEDITORTOOLSEDITORONLY_API UUVEditorParameterizeMeshOperatorFactory : public UObject, public UE::Geometry::IDynamicMeshOperatorFactory
+{
+	GENERATED_BODY()
+
+public:
+	// IDynamicMeshOperatorFactory API
+	virtual TUniquePtr<UE::Geometry::FDynamicMeshOperator> MakeNewOperator() override;
+
+	UPROPERTY()
+	TObjectPtr<UUVEditorParameterizeMeshToolProperties> Settings = nullptr;
+
+	UPROPERTY()
+	TObjectPtr<UUVEditorParameterizeMeshToolUVAtlasProperties> UVAtlasProperties = nullptr;
+
+	UPROPERTY()
+	TObjectPtr<UUVEditorParameterizeMeshToolXAtlasProperties> XAtlasProperties = nullptr;
+
+	UPROPERTY()
+	TObjectPtr<UUVEditorParameterizeMeshToolPatchBuilderProperties> PatchBuilderProperties = nullptr;
+
+	TSharedPtr<UE::Geometry::FPolygroupSet, ESPMode::ThreadSafe> InputGroups;
+	TSharedPtr<UE::Geometry::FDynamicMesh3, ESPMode::ThreadSafe> OriginalMesh;
+	TUniqueFunction<int32()> GetSelectedUVChannel = []() { return 0; };
+	FTransform TargetTransform;
+};
\ No newline at end of file
diff --git a/Engine/Plugins/Experimental/UVEditor/Source/UVEditorToolsEditorOnly/Public/UVEditorParameterizeMeshTool.h b/Engine/Plugins/Experimental/UVEditor/Source/UVEditorToolsEditorOnly/Public/UVEditorParameterizeMeshTool.h
--- a/Engine/Plugins/Experimental/UVEditor/Source/UVEditorToolsEditorOnly/Public/UVEditorParameterizeMeshTool.h
+++ b/Engine/Plugins/Experimental/UVEditor/Source/UVEditorToolsEditorOnly/Public/UVEditorParameterizeMeshTool.h
@@ -12,6 +12,7 @@
 #include "Properties/MeshMaterialProperties.h"
 #include "Properties/MeshUVChannelProperties.h"
 #include "UVEditorToolAnalyticsUtils.h"
+#include "Polygroups/PolygroupSet.h"
 
 #include "UVEditorParameterizeMeshTool.generated.h"
 
@@ -20,11 +21,12 @@ class UDynamicMeshComponent;
 class UMaterialInterface;
 class UMaterialInstanceDynamic;
 class UUVEditorToolMeshInput;
-class UParameterizeMeshToolProperties;
-class UParameterizeMeshToolUVAtlasProperties;
-class UParameterizeMeshToolXAtlasProperties;
-class UParameterizeMeshToolPatchBuilderProperties;
-class UParameterizeMeshOperatorFactory;
+class UUVEditorParameterizeMeshToolProperties;
+class UUVEditorParameterizeMeshToolUVAtlasProperties;
+class UUVEditorParameterizeMeshToolXAtlasProperties;
+class UUVEditorParameterizeMeshToolPatchBuilderProperties;
+class UUVEditorParameterizeMeshOperatorFactory;
+class UPolygroupLayersProperties;
 
 UCLASS()
 class UVEDITORTOOLSEDITORONLY_API UUVEditorParameterizeMeshToolBuilder : public UInteractiveToolBuilder
@@ -75,22 +77,29 @@ protected:
 	TArray<TObjectPtr<UUVEditorToolMeshInput>> Targets;
 
 	UPROPERTY()
-	TObjectPtr<UParameterizeMeshToolProperties> Settings = nullptr;
+	TObjectPtr<UUVEditorParameterizeMeshToolProperties> Settings = nullptr;
 
 	UPROPERTY()
-	TObjectPtr<UParameterizeMeshToolUVAtlasProperties> UVAtlasProperties = nullptr;
+	TObjectPtr<UUVEditorParameterizeMeshToolUVAtlasProperties> UVAtlasProperties = nullptr;
 
 	UPROPERTY()
-	TObjectPtr<UParameterizeMeshToolXAtlasProperties> XAtlasProperties = nullptr;
+	TObjectPtr<UUVEditorParameterizeMeshToolXAtlasProperties> XAtlasProperties = nullptr;
 
 	UPROPERTY()
-	TObjectPtr<UParameterizeMeshToolPatchBuilderProperties> PatchBuilderProperties = nullptr;
+	TObjectPtr<UUVEditorParameterizeMeshToolPatchBuilderProperties> PatchBuilderProperties = nullptr;
 
 	UPROPERTY()
-	TArray<TObjectPtr<UParameterizeMeshOperatorFactory>> Factories;
+	TObjectPtr<UPolygroupLayersProperties> PolygroupLayerProperties = nullptr;
+
+	UPROPERTY()
+	TArray<TObjectPtr<UUVEditorParameterizeMeshOperatorFactory>> Factories;
 
 	void OnMethodTypeChanged();
 
+	TSharedPtr<UE::Geometry::FPolygroupSet, ESPMode::ThreadSafe> ActiveGroupSet;
+	void OnSelectedGroupLayerChanged();
+	void UpdateActiveGroupLayer(bool bUpdateFactories = true);
+
 	//
 	// Analytics
 	//
diff --git a/Engine/Plugins/Experimental/UVEditor/Source/UVEditorToolsEditorOnly/UVEditorToolsEditorOnly.Build.cs b/Engine/Plugins/Experimental/UVEditor/Source/UVEditorToolsEditorOnly/UVEditorToolsEditorOnly.Build.cs
--- a/Engine/Plugins/Experimental/UVEditor/Source/UVEditorToolsEditorOnly/UVEditorToolsEditorOnly.Build.cs
+++ b/Engine/Plugins/Experimental/UVEditor/Source/UVEditorToolsEditorOnly/UVEditorToolsEditorOnly.Build.cs
@@ -58,11 +58,27 @@ public class UVEditorToolsEditorOnly : ModuleRules
 				"EditorWidgets",
 				"ModelingEditorUI",
 				"PropertyEditor",
+
+				"MeshDescription",
+				"StaticMeshDescription",
+				"MeshConversion",
+				"GeometryAlgorithms", // required for constrained Delaunay triangulation
+				"MeshUtilitiesCommon", // required by uvlayoutop
 				// ... add private dependencies that you statically link with here ...	
 			}
 			);
-		
-		
+		bool bWithProxyLOD = Target.Platform == UnrealTargetPlatform.Win64;
+		PrivateDefinitions.Add("WITH_PROXYLOD=" + (bWithProxyLOD ? '1' : '0'));
+		if (bWithProxyLOD)
+		{
+			PrivateDependencyModuleNames.AddRange(
+				new string[]
+				{
+					"ProxyLODMeshReduction", // currently Win64-only
+				}
+				);
+		}
+
 		DynamicallyLoadedModuleNames.AddRange(
 			new string[]
 			{
diff --git a/Engine/Plugins/Runtime/GeometryProcessing/Source/DynamicMesh/Private/Parameterization/MeshRegionGraph.cpp b/Engine/Plugins/Runtime/GeometryProcessing/Source/DynamicMesh/Private/Parameterization/MeshRegionGraph.cpp
--- a/Engine/Plugins/Runtime/GeometryProcessing/Source/DynamicMesh/Private/Parameterization/MeshRegionGraph.cpp
+++ b/Engine/Plugins/Runtime/GeometryProcessing/Source/DynamicMesh/Private/Parameterization/MeshRegionGraph.cpp
@@ -45,7 +45,10 @@ void FMeshRegionGraph::BuildFromComponents(
 }
 
 
-void FMeshRegionGraph::BuildFromTriangleSets(const FDynamicMesh3& Mesh, const TArray<TArray<int32>>& TriangleSets, TFunctionRef<int32(int32)> ExternalIDFunc)
+void FMeshRegionGraph::BuildFromTriangleSets(const FDynamicMesh3& Mesh,
+	const TArray<TArray<int32>>& TriangleSets,
+	TFunctionRef<int32(int32)> ExternalIDFunc,
+	TFunctionRef<bool(int32, int32)> TrisConnectedPredicate)
 {
 	int32 N = TriangleSets.Num();
 	Regions.SetNum(N);
@@ -60,7 +63,15 @@ void FMeshRegionGraph::BuildFromTriangleSets(const FDynamicMesh3& Mesh, const TA
 		for (int32 tid : TriangleSets[k])
 		{
 			TriangleToRegionMap[tid] = k;
-			TriangleNbrTris[tid] = Mesh.GetTriNeighbourTris(tid);
+			FIndex3i NbrTris = Mesh.GetTriNeighbourTris(tid);
+			for (int j = 0; j < 3; ++j)
+			{
+				if (NbrTris[j] >= 0 && TrisConnectedPredicate(tid, NbrTris[j]) == false)
+				{
+					NbrTris[j] = -1;
+				}
+			}
+			TriangleNbrTris[tid] = NbrTris;
 		}
 	}
 
diff --git a/Engine/Plugins/Runtime/GeometryProcessing/Source/DynamicMesh/Private/Polygroups/PolygroupsGenerator.cpp b/Engine/Plugins/Runtime/GeometryProcessing/Source/DynamicMesh/Private/Polygroups/PolygroupsGenerator.cpp
--- a/Engine/Plugins/Runtime/GeometryProcessing/Source/DynamicMesh/Private/Polygroups/PolygroupsGenerator.cpp
+++ b/Engine/Plugins/Runtime/GeometryProcessing/Source/DynamicMesh/Private/Polygroups/PolygroupsGenerator.cpp
@@ -252,7 +252,8 @@ bool FPolygroupsGenerator::FindPolygroupsFromFurthestPointSampling(
 	NumPoints = FMath::Min(NumPoints, Mesh->VertexCount());
 
 	// cannot seem to use auto or TUniqueFunction here...
-	TFunction<bool(int,int)> TrisConnectedPredicate = [](int, int) -> bool { return true; };
+	TFunction<bool(int32, int32)> TrisConnectedPredicate;
+	TrisConnectedPredicate = [](int, int) -> bool { return true; };
 	if (StartingGroups != nullptr)
 	{
 		TrisConnectedPredicate = [StartingGroups](int a, int b) -> bool { return StartingGroups->GetGroup(a) == StartingGroups->GetGroup(b) ? true : false; };
@@ -375,7 +376,7 @@ bool FPolygroupsGenerator::FindPolygroupsFromFurthestPointSampling(
 
 	if (bApplyPostProcessing)
 	{
-		PostProcessPolygroups(true);
+		PostProcessPolygroups(true, TrisConnectedPredicate);
 	}
 	if (bCopyToMesh)
 	{
@@ -388,19 +389,20 @@ bool FPolygroupsGenerator::FindPolygroupsFromFurthestPointSampling(
 
 
 
-void FPolygroupsGenerator::PostProcessPolygroups(bool bApplyMerging)
+void FPolygroupsGenerator::PostProcessPolygroups(bool bApplyMerging, TFunctionRef<bool(int32, int32)> TrisConnectedPredicate)
 {
 	if (bApplyMerging && MinGroupSize > 1)
 	{
-		OptimizePolygroups();
+		OptimizePolygroups(TrisConnectedPredicate);
 	}
 }
 
 
-void FPolygroupsGenerator::OptimizePolygroups()
+void FPolygroupsGenerator::OptimizePolygroups(TFunctionRef<bool(int32, int32)> TrisConnectedPredicate)
 {
 	FMeshRegionGraph RegionGraph;
-	RegionGraph.BuildFromTriangleSets(*Mesh, FoundPolygroups, [&](int32 SetIdx) { return SetIdx; });
+	RegionGraph.BuildFromTriangleSets(*Mesh, FoundPolygroups, [&](int32 SetIdx) { return SetIdx; },
+		                                                      TrisConnectedPredicate);		                                                      
 	bool bMerged = RegionGraph.MergeSmallRegions(MinGroupSize-1, [&](int32 A, int32 B) { return RegionGraph.GetRegionTriCount(A) > RegionGraph.GetRegionTriCount(B); });
 	bool bSwapped = RegionGraph.OptimizeBorders();
 	if (bMerged || bSwapped)
diff --git a/Engine/Plugins/Runtime/GeometryProcessing/Source/DynamicMesh/Public/Parameterization/MeshRegionGraph.h b/Engine/Plugins/Runtime/GeometryProcessing/Source/DynamicMesh/Public/Parameterization/MeshRegionGraph.h
--- a/Engine/Plugins/Runtime/GeometryProcessing/Source/DynamicMesh/Public/Parameterization/MeshRegionGraph.h
+++ b/Engine/Plugins/Runtime/GeometryProcessing/Source/DynamicMesh/Public/Parameterization/MeshRegionGraph.h
@@ -34,7 +34,7 @@ public:
 	void BuildFromComponents(const FDynamicMesh3& Mesh, 
 							 const FMeshConnectedComponents& Components,
 							 TFunctionRef<int32(int32)> ExternalIDFunc,
-							 TFunctionRef<bool(int32, int32)> TrisConnectedPredicate = [](int a, int b) { return true; } );
+							 TFunctionRef<bool(int32, int32)> TrisConnectedPredicate = [](int, int) { return true; } );
 
 	/**
 	 * Build a region graph for a Mesh using the given triangle sets
@@ -42,7 +42,8 @@ public:
 	 */
 	void BuildFromTriangleSets(const FDynamicMesh3& Mesh,
 							 const TArray<TArray<int32>>& TriangleSets,
-							 TFunctionRef<int32(int32)> ExternalIDFunc);
+							 TFunctionRef<int32(int32)> ExternalIDFunc,
+		                     TFunctionRef<bool(int32, int32)> TrisConnectedPredicate = [](int, int) { return true; });
 
 	/** 
 	 * @return Max valid region index. Some region indices may be invalid. 
diff --git a/Engine/Plugins/Runtime/GeometryProcessing/Source/DynamicMesh/Public/Polygroups/PolygroupsGenerator.h b/Engine/Plugins/Runtime/GeometryProcessing/Source/DynamicMesh/Public/Polygroups/PolygroupsGenerator.h
--- a/Engine/Plugins/Runtime/GeometryProcessing/Source/DynamicMesh/Public/Polygroups/PolygroupsGenerator.h
+++ b/Engine/Plugins/Runtime/GeometryProcessing/Source/DynamicMesh/Public/Polygroups/PolygroupsGenerator.h
@@ -115,8 +115,8 @@ public:
 
 
 protected:
-	void PostProcessPolygroups(bool bApplyMerging);
-	void OptimizePolygroups();
+	void PostProcessPolygroups(bool bApplyMerging, TFunctionRef<bool(int32, int32)> TrisConnectedPredicate = [](int, int) { return true; });
+	void OptimizePolygroups(TFunctionRef<bool(int32, int32)> TrisConnectedPredicate = [](int, int) { return true; });
 };
 
 
