UVEditor: Adds new transform tool, for moving, aligning, and distributing selected islands and triangles in the UV space

diff --git a/Engine/Plugins/Experimental/UVEditor/Source/UVEditor/Private/UVEditorCommands.cpp b/Engine/Plugins/Experimental/UVEditor/Source/UVEditor/Private/UVEditorCommands.cpp
--- a/Engine/Plugins/Experimental/UVEditor/Source/UVEditor/Private/UVEditorCommands.cpp
+++ b/Engine/Plugins/Experimental/UVEditor/Source/UVEditor/Private/UVEditorCommands.cpp
@@ -24,7 +24,10 @@ void FUVEditorCommands::RegisterCommands()
 	UI_COMMAND(ApplyChanges, "Apply", "Apply changes to original meshes", EUserInterfaceActionType::Button, FInputChord());
 
 	// These get linked to various tool buttons.
-	UI_COMMAND(BeginLayoutTool, "Layout", "Transform and pack existing UVs", EUserInterfaceActionType::ToggleButton, FInputChord());
+	UI_COMMAND(BeginLayoutTool, "Layout", "Pack existing UVs", EUserInterfaceActionType::ToggleButton, FInputChord());
+	UI_COMMAND(BeginTransformTool, "Transform", "Transform existing UVs", EUserInterfaceActionType::ToggleButton, FInputChord());
+	UI_COMMAND(BeginAlignTool, "Align", "Align existing UVs", EUserInterfaceActionType::ToggleButton, FInputChord());
+	UI_COMMAND(BeginDistributeTool, "Distribute", "Distribute existing UVs", EUserInterfaceActionType::ToggleButton, FInputChord());
 	UI_COMMAND(BeginParameterizeMeshTool, "AutoUV", "Auto-unwrap and pack UVs", EUserInterfaceActionType::ToggleButton, FInputChord());
 	UI_COMMAND(BeginChannelEditTool, "Channels", "Modify UV channels", EUserInterfaceActionType::ToggleButton, FInputChord());
 	UI_COMMAND(BeginSeamTool, "Seam", "Add UV seams", EUserInterfaceActionType::ToggleButton, FInputChord());
diff --git a/Engine/Plugins/Experimental/UVEditor/Source/UVEditor/Private/UVEditorMode.cpp b/Engine/Plugins/Experimental/UVEditor/Source/UVEditor/Private/UVEditorMode.cpp
--- a/Engine/Plugins/Experimental/UVEditor/Source/UVEditor/Private/UVEditorMode.cpp
+++ b/Engine/Plugins/Experimental/UVEditor/Source/UVEditor/Private/UVEditorMode.cpp
@@ -33,6 +33,7 @@
 #include "ToolTargets/UVEditorToolMeshInput.h"
 #include "UVEditorCommands.h"
 #include "UVEditorLayoutTool.h"
+#include "UVEditorTransformTool.h"
 #include "UVEditorParameterizeMeshTool.h"
 #include "UVEditorLayerEditTool.h"
 #include "UVEditorSeamTool.h"
@@ -370,6 +371,18 @@ void UUVEditorMode::RegisterTools()
 	UVEditorLayoutToolBuilder->Targets = &ToolInputObjects;
 	RegisterTool(CommandInfos.BeginLayoutTool, TEXT("BeginLayoutTool"), UVEditorLayoutToolBuilder);
 
+	UUVEditorTransformToolBuilder* UVEditorTransformToolBuilder = NewObject<UUVEditorTransformToolBuilder>();
+	UVEditorTransformToolBuilder->Targets = &ToolInputObjects;
+	RegisterTool(CommandInfos.BeginTransformTool, TEXT("BeginTransformTool"), UVEditorTransformToolBuilder);
+
+	UUVEditorAlignToolBuilder* UVEditorAlignToolBuilder = NewObject<UUVEditorAlignToolBuilder>();
+	UVEditorAlignToolBuilder->Targets = &ToolInputObjects;
+	RegisterTool(CommandInfos.BeginAlignTool, TEXT("BeginAlignTool"), UVEditorAlignToolBuilder);
+
+	UUVEditorDistributeToolBuilder* UVEditorDistributeToolBuilder = NewObject<UUVEditorDistributeToolBuilder>();
+	UVEditorDistributeToolBuilder->Targets = &ToolInputObjects;
+	RegisterTool(CommandInfos.BeginDistributeTool, TEXT("BeginDistributeTool"), UVEditorDistributeToolBuilder);
+
 	UUVEditorParameterizeMeshToolBuilder* UVEditorParameterizeMeshToolBuilder = NewObject<UUVEditorParameterizeMeshToolBuilder>();
 	UVEditorParameterizeMeshToolBuilder->Targets = &ToolInputObjects;
 	RegisterTool(CommandInfos.BeginParameterizeMeshTool, TEXT("BeginParameterizeMeshTool"), UVEditorParameterizeMeshToolBuilder);
diff --git a/Engine/Plugins/Experimental/UVEditor/Source/UVEditor/Private/UVEditorModeToolkit.cpp b/Engine/Plugins/Experimental/UVEditor/Source/UVEditor/Private/UVEditorModeToolkit.cpp
--- a/Engine/Plugins/Experimental/UVEditor/Source/UVEditor/Private/UVEditorModeToolkit.cpp
+++ b/Engine/Plugins/Experimental/UVEditor/Source/UVEditor/Private/UVEditorModeToolkit.cpp
@@ -383,6 +383,9 @@ void FUVEditorModeToolkit::BuildToolPalette(FName PaletteIndex, class FToolBarBu
 		ToolbarBuilder.AddToolBarButton(Commands.SplitAction);
 
 		ToolbarBuilder.AddToolBarButton(Commands.BeginLayoutTool);
+		ToolbarBuilder.AddToolBarButton(Commands.BeginTransformTool);
+		ToolbarBuilder.AddToolBarButton(Commands.BeginAlignTool);
+		ToolbarBuilder.AddToolBarButton(Commands.BeginDistributeTool);
 		ToolbarBuilder.AddToolBarButton(Commands.BeginChannelEditTool);
 		ToolbarBuilder.AddToolBarButton(Commands.BeginSeamTool);
 		ToolbarBuilder.AddToolBarButton(Commands.BeginParameterizeMeshTool);
diff --git a/Engine/Plugins/Experimental/UVEditor/Source/UVEditor/Public/UVEditorCommands.h b/Engine/Plugins/Experimental/UVEditor/Source/UVEditor/Public/UVEditorCommands.h
--- a/Engine/Plugins/Experimental/UVEditor/Source/UVEditor/Public/UVEditorCommands.h
+++ b/Engine/Plugins/Experimental/UVEditor/Source/UVEditor/Public/UVEditorCommands.h
@@ -19,6 +19,9 @@ public:
 	TSharedPtr<FUICommandInfo> ApplyChanges;
 
 	TSharedPtr<FUICommandInfo> BeginLayoutTool;
+	TSharedPtr<FUICommandInfo> BeginTransformTool;
+	TSharedPtr<FUICommandInfo> BeginAlignTool;
+	TSharedPtr<FUICommandInfo> BeginDistributeTool;
 	TSharedPtr<FUICommandInfo> BeginParameterizeMeshTool;
 	TSharedPtr<FUICommandInfo> BeginChannelEditTool;
 	TSharedPtr<FUICommandInfo> BeginSeamTool;
diff --git a/Engine/Plugins/Experimental/UVEditor/Source/UVEditorTools/Private/Operators/UVEditorUVTransformOp.cpp b/Engine/Plugins/Experimental/UVEditor/Source/UVEditorTools/Private/Operators/UVEditorUVTransformOp.cpp
new file mode 100644
--- /dev/null
+++ b/Engine/Plugins/Experimental/UVEditor/Source/UVEditorTools/Private/Operators/UVEditorUVTransformOp.cpp
@@ -0,0 +1,646 @@
+// Copyright Epic Games, Inc. All Rights Reserved.
+
+#include "Operators/UVEditorUVTransformOp.h"
+
+#include "DynamicMesh/DynamicMesh3.h"
+#include "DynamicMesh/DynamicMeshAttributeSet.h"
+#include "Parameterization/MeshUVPacking.h"
+#include "Selections/MeshConnectedComponents.h"
+#include "Utilities/MeshUDIMClassifier.h"
+#include "Async/ParallelFor.h"
+#include "UVEditorUXSettings.h"
+
+using namespace UE::Geometry;
+
+void FUVEditorUVTransformBaseOp::SetTransform(const FTransformSRT3d& Transform)
+{
+	ResultTransform = Transform;
+}
+
+void FUVEditorUVTransformBaseOp::RebuildBoundingBoxes()
+{	
+	int32 NumComponents = UVComponents->Num();
+	PerComponentBoundingBoxes.SetNum(NumComponents);
+	OverallBoundingBox = FAxisAlignedBox2d::Empty();
+
+	ParallelFor(NumComponents, [&](int32 k)
+		{
+			FVector2d TriangleUV[3];
+			PerComponentBoundingBoxes[k] = FAxisAlignedBox2d::Empty();
+			const TArray<int>& Triangles = (*UVComponents)[k].Indices;
+			for (int tid : Triangles)
+			{
+				FIndex3i TriElements = ActiveUVLayer->GetTriangle(tid);
+				if (!TriElements.Contains(FDynamicMesh3::InvalidID))
+				{
+					for (int j = 0; j < 3; ++j)
+					{
+						TriangleUV[j] = FVector2d(FUVEditorUXSettings::InternalUVToExternalUV(ActiveUVLayer->GetElement(TriElements[j])));
+						PerComponentBoundingBoxes[k].Contain(TriangleUV[j]);
+					}
+				}
+			}
+		});
+	for (int32 Cid = 0; Cid < NumComponents; ++Cid)
+	{
+		OverallBoundingBox.Contain(PerComponentBoundingBoxes[Cid]);
+	}
+}
+
+void FUVEditorUVTransformBaseOp::CollectTransformElements()
+{
+	if (Selection.IsSet())
+	{
+		TransformingElements = TSet<int32>();
+	}
+
+	int32 NumComponents = UVComponents->Num();
+
+	for (int32 k = 0; k < NumComponents; ++k)
+	{
+		const TArray<int>& Triangles = (*UVComponents)[k].Indices;
+		for (int tid : Triangles)
+		{
+			FIndex3i Elements = ActiveUVLayer->GetTriangle(tid);
+			for (int j = 0; j < 3; ++j)
+			{
+				if (TransformingElements.IsSet())
+				{
+					TransformingElements.GetValue().Add(Elements[j]);
+				}
+				ElementToComponent.Add(Elements[j], k);
+			}
+		}
+	}
+}
+
+FVector2f FUVEditorUVTransformBaseOp::GetAlignmentPointFromBoundingBoxAndDirection(EUVEditorAlignDirectionBackend Direction, const FAxisAlignedBox2d& BoundingBox)
+{
+	FVector2f AlignmentPoint = FVector2f();
+
+	switch (Direction)
+	{
+	case EUVEditorAlignDirectionBackend::Top:
+		AlignmentPoint = FVector2f(BoundingBox.Center().X, BoundingBox.Max.Y);
+		break;
+	case EUVEditorAlignDirectionBackend::Bottom:
+		AlignmentPoint = FVector2f(BoundingBox.Center().X, BoundingBox.Min.Y);
+		break;
+	case EUVEditorAlignDirectionBackend::Left:
+		AlignmentPoint = FVector2f(BoundingBox.Min.X, BoundingBox.Center().Y);
+		break;
+	case EUVEditorAlignDirectionBackend::Right:
+		AlignmentPoint = FVector2f(BoundingBox.Max.X, BoundingBox.Center().Y);
+		break;
+	case EUVEditorAlignDirectionBackend::CenterVertically:
+		AlignmentPoint = FVector2f(BoundingBox.Center().X, BoundingBox.Center().Y);
+		break;
+	case EUVEditorAlignDirectionBackend::CenterHorizontally:
+		AlignmentPoint = FVector2f(BoundingBox.Center().X, BoundingBox.Center().Y);
+		break;
+	default:
+		ensure(false);
+	}
+
+	return AlignmentPoint;
+}
+
+FVector2f FUVEditorUVTransformBaseOp::GetAlignmentPointFromUDIMAndDirection(EUVEditorAlignDirectionBackend Direction, FVector2i UDIMTile)
+{
+	FVector2f AlignmentPoint = FVector2f();
+
+	FVector2f UDIMLowerCorner(UDIMTile);
+
+	switch (Direction)
+	{
+	case EUVEditorAlignDirectionBackend::Top:
+		AlignmentPoint = FVector2f(UDIMLowerCorner.X + 0.5f, UDIMLowerCorner.Y + 1.0f);
+		break;
+	case EUVEditorAlignDirectionBackend::Bottom:
+		AlignmentPoint = FVector2f(UDIMLowerCorner.X + 0.5f, UDIMLowerCorner.Y);
+		break;
+	case EUVEditorAlignDirectionBackend::Left:
+		AlignmentPoint = FVector2f(UDIMLowerCorner.X, UDIMLowerCorner.Y + 0.5f);
+		break;
+	case EUVEditorAlignDirectionBackend::Right:
+		AlignmentPoint = FVector2f(UDIMLowerCorner.X + 1.0f, UDIMLowerCorner.Y + 0.5f);
+		break;
+	case EUVEditorAlignDirectionBackend::CenterVertically:
+		AlignmentPoint = FVector2f(UDIMLowerCorner.X + 0.5f, UDIMLowerCorner.Y + 0.5f);
+		break;
+	case EUVEditorAlignDirectionBackend::CenterHorizontally:
+		AlignmentPoint = FVector2f(UDIMLowerCorner.X + 0.5f, UDIMLowerCorner.Y + 0.5f);
+		break;
+	default:
+		ensure(false);
+	}
+
+	return AlignmentPoint;
+}
+
+void FUVEditorUVTransformBaseOp::CalculateResult(FProgressCancel* Progress)
+{
+	if (Progress && Progress->Cancelled())
+	{
+		return;
+	}
+
+	ResultMesh->Copy(*OriginalMesh, true, true, true, true);
+	if (!ensureMsgf(ResultMesh->HasAttributes(), TEXT("Attributes not found on mesh? Conversion should always create them, so this operator should not need to do so.")))
+	{
+		ResultMesh->EnableAttributes();
+	}
+
+	if (Progress && Progress->Cancelled())
+	{
+		return;
+	}
+
+	int UVLayerInput = UVLayerIndex;
+	ActiveUVLayer = ResultMesh->Attributes()->GetUVLayer(UVLayerInput);
+
+	auto UVIslandPredicate = [this](int32 Triangle0, int32 Triangle1)
+	{
+		return ActiveUVLayer->AreTrianglesConnected(Triangle0, Triangle1);
+	};
+
+	UVComponents = MakeShared<FMeshConnectedComponents>(ResultMesh.Get());
+	if (Selection.IsSet())
+	{
+		UVComponents->FindConnectedTriangles(Selection.GetValue().Array(), UVIslandPredicate);
+	}
+	else
+	{
+		UVComponents->FindConnectedTriangles(UVIslandPredicate);
+	}
+
+	if (Progress && Progress->Cancelled())
+	{
+		return;
+	}
+
+	CollectTransformElements();
+
+	if (Progress && Progress->Cancelled())
+	{
+		return;
+	}
+
+	HandleTransformationOp(Progress);
+}
+
+FVector2f FUVEditorUVTransformOp::GetPivotFromMode(int32 ElementID)
+{
+	FVector2f Pivot;
+	const int32* Component;
+
+	switch (PivotMode)
+	{
+	case EUVEditorPivotTypeBackend::Origin:
+		Pivot = FVector2f(0, 0);
+		break;
+	case EUVEditorPivotTypeBackend::BoundingBoxCenter:
+		Pivot = FVector2f(OverallBoundingBox.Center());
+		break;
+	case EUVEditorPivotTypeBackend::IndividualBoundingBoxCenter:
+		Component = ElementToComponent.Find(ElementID);
+		if (ensure(Component != nullptr))
+		{
+			Pivot = FVector2f(PerComponentBoundingBoxes[*Component].Center());
+		}
+		break;
+	case EUVEditorPivotTypeBackend::Manual:
+		Pivot = FVector2f(ManualPivot);
+		break;
+	}
+	return Pivot;
+}
+
+void FUVEditorUVTransformOp::HandleTransformationOp(FProgressCancel * Progress)
+{
+	auto ScaleFunc = [&](int32 Vid)
+	{
+		FVector2f ScalePivot = GetPivotFromMode(Vid);
+		FVector2f UV = FUVEditorUXSettings::InternalUVToExternalUV(ActiveUVLayer->GetElement(Vid));
+		UV = (UV - ScalePivot);
+		UV[0] *= Scale[0];
+		UV[1] *= Scale[1];
+		UV = (UV + ScalePivot);
+		ActiveUVLayer->SetElement(Vid, FUVEditorUXSettings::ExternalUVToInternalUV(UV));
+	};
+
+	auto RotFunc = [&](int32 Vid)
+	{
+		FVector2f RotationPivot = GetPivotFromMode(Vid);
+		FVector2f UV = FUVEditorUXSettings::InternalUVToExternalUV(ActiveUVLayer->GetElement(Vid));
+		FVector2f UV_Rotated;
+		double RotationInRadians = Rotation / 180.0 * UE_PI;
+		UV = (UV - RotationPivot);
+		UV_Rotated[0] = UV[0] * FMath::Cos(RotationInRadians) - UV[1] * FMath::Sin(RotationInRadians);
+		UV_Rotated[1] = UV[0] * FMath::Sin(RotationInRadians) + UV[1] * FMath::Cos(RotationInRadians);
+		UV = (UV_Rotated + RotationPivot);
+		ActiveUVLayer->SetElement(Vid, FUVEditorUXSettings::ExternalUVToInternalUV(UV));
+	};
+
+	auto TranslateFunc = [&](int32 Vid)
+	{	
+		FVector2f RotationPivot = FVector2f(0,0);
+		if (TranslationMode == EUVEditorTranslationModeBackend::Absolute)
+		{
+			RotationPivot = GetPivotFromMode(Vid);
+		}				
+		FVector2f UV = FUVEditorUXSettings::InternalUVToExternalUV(ActiveUVLayer->GetElement(Vid));
+		UV = (UV + FVector2f(Translation) - RotationPivot);
+		ActiveUVLayer->SetElement(Vid, FUVEditorUXSettings::ExternalUVToInternalUV(UV));
+	};
+
+	auto ApplyTransformFunc = [&](TFunction<void(int32 Vid)> TransformFunc)
+	{
+		RebuildBoundingBoxes();
+
+		if (TransformingElements.IsSet())
+		{
+			for (int ElementID : TransformingElements.GetValue())
+			{
+				TransformFunc(ElementID);
+
+				if (Progress && Progress->Cancelled())
+				{
+					return;
+				}
+			}
+		}
+		else
+		{
+			for (int ElementID : ActiveUVLayer->ElementIndicesItr())
+			{
+				TransformFunc(ElementID);
+
+				if (Progress && Progress->Cancelled())
+				{
+					return;
+				}
+			}
+		}
+	};
+
+	if (!FMath::IsNearlyEqual(Scale.X, 1.0f) || !FMath::IsNearlyEqual(Scale.Y, 1.0f))
+	{
+		ApplyTransformFunc(ScaleFunc);
+	}
+	if (!FMath::IsNearlyZero(Rotation))
+	{
+		ApplyTransformFunc(RotFunc);
+	}
+	if (!FMath::IsNearlyZero(Translation.X) || !FMath::IsNearlyZero(Translation.Y)
+		|| TranslationMode == EUVEditorTranslationModeBackend::Absolute)
+	{
+		ApplyTransformFunc(TranslateFunc);
+	}
+
+}
+
+
+void FUVEditorUVAlignOp::HandleTransformationOp(FProgressCancel* Progress)
+{
+	int32 NumComponents = UVComponents->Num();
+	RebuildBoundingBoxes();
+
+	auto TranslateFunc = [&](const FVector2f& Translation, int32 Vid)
+	{
+		FVector2f UV = FUVEditorUXSettings::InternalUVToExternalUV(ActiveUVLayer->GetElement(Vid));
+		UV = (UV + Translation);
+		ActiveUVLayer->SetElement(Vid, FUVEditorUXSettings::ExternalUVToInternalUV(UV));
+	};
+
+	auto TranslationFromAlignmentPoints = [this](const FVector2f& PointTo, const FVector2f& PointFrom)
+	{
+		switch (AlignDirection)
+		{
+		case EUVEditorAlignDirectionBackend::Top:
+			return FVector2f(0, PointTo.Y - PointFrom.Y);
+		case EUVEditorAlignDirectionBackend::Bottom:
+			return FVector2f(0, PointTo.Y - PointFrom.Y);
+		case EUVEditorAlignDirectionBackend::Left:
+			return FVector2f(PointTo.X - PointFrom.X, 0);
+		case EUVEditorAlignDirectionBackend::Right:
+			return FVector2f(PointTo.X - PointFrom.X, 0);
+		case EUVEditorAlignDirectionBackend::CenterVertically:
+			return FVector2f(PointTo.X - PointFrom.X, 0);
+		case EUVEditorAlignDirectionBackend::CenterHorizontally:
+			return FVector2f(0, PointTo.Y - PointFrom.Y);
+		default:
+			ensure(false);
+			return FVector2f(0, 0);
+		}
+	};
+
+	TArray<FVector2f> PerComponentTranslation;
+	PerComponentTranslation.SetNum(NumComponents);
+	for (int32 Cid = 0; Cid < NumComponents; ++Cid)
+	{
+		FVector2f ComponentAlignmentPoint = GetAlignmentPointFromBoundingBoxAndDirection(AlignDirection, PerComponentBoundingBoxes[Cid]);
+		switch (AlignAnchor)
+		{
+		case EUVEditorAlignAnchorBackend::UDIMTile:
+		{
+			FVector2i UDIM = FDynamicMeshUDIMClassifier::ClassifyTrianglesToUDIM(ActiveUVLayer, UVComponents->Components[Cid].Indices);
+			FVector2f TileAlignmentPoint = GetAlignmentPointFromUDIMAndDirection(AlignDirection, UDIM);
+			PerComponentTranslation[Cid] = TranslationFromAlignmentPoints(TileAlignmentPoint, ComponentAlignmentPoint);
+		}
+			break;
+		case EUVEditorAlignAnchorBackend::BoundingBox:
+		{		
+			FVector2f BoundingBoxAlignmentPoint = GetAlignmentPointFromBoundingBoxAndDirection(AlignDirection, OverallBoundingBox);
+			PerComponentTranslation[Cid] = TranslationFromAlignmentPoints(BoundingBoxAlignmentPoint, ComponentAlignmentPoint);
+		}
+		break;
+		case EUVEditorAlignAnchorBackend::Manual:
+		{
+			PerComponentTranslation[Cid] = TranslationFromAlignmentPoints((FVector2f)ManualAnchor, ComponentAlignmentPoint);
+		}
+		break;
+		default:
+			ensure(false);
+			PerComponentTranslation[Cid] = FVector2f();
+			break;
+		}
+	}
+
+	if (TransformingElements.IsSet())
+	{
+		for (int ElementID : TransformingElements.GetValue())
+		{
+			TranslateFunc(PerComponentTranslation[*ElementToComponent.Find(ElementID)], ElementID);
+			if (Progress && Progress->Cancelled())
+			{
+				return;
+			}
+		}
+	}
+	else
+	{
+		for (int ElementID : ActiveUVLayer->ElementIndicesItr())
+		{
+			TranslateFunc(PerComponentTranslation[*ElementToComponent.Find(ElementID)], ElementID);
+			if (Progress && Progress->Cancelled())
+			{
+				return;
+			}
+		}
+	}
+	
+
+}
+
+
+void FUVEditorUVDistributeOp::HandleTransformationOp(FProgressCancel* Progress)
+{
+	int32 NumComponents = UVComponents->Num();
+	RebuildBoundingBoxes();
+	TArray<FVector2f> PerComponentTranslation;
+	
+
+	auto ComputeDistributeTranslations = [this, NumComponents](bool bVertical, EUVEditorAlignDirectionBackend EdgeDirection, float SpreadDirection, bool bEqualizeSpacing)
+	{
+		TArray<FVector2f> PerComponentTranslation;
+		PerComponentTranslation.SetNum(NumComponents);
+		float TotalDistance = 0.0f;
+		for (int32 Cid = 0; Cid < NumComponents; ++Cid)
+		{
+			TotalDistance += bVertical ? PerComponentBoundingBoxes[Cid].Height() : PerComponentBoundingBoxes[Cid].Width();
+		}
+		float BoundingBoxDistance = bVertical ? OverallBoundingBox.Height() : OverallBoundingBox.Width();
+		float GapSpace = (BoundingBoxDistance - TotalDistance) / (NumComponents-1);
+		float PerComponentSpace = BoundingBoxDistance / (NumComponents);
+
+		float NextPosition = 0.0f;
+		FVector2f OverallAlignmentPoint = GetAlignmentPointFromBoundingBoxAndDirection(EdgeDirection, OverallBoundingBox);
+		for (int32 Cid = 0; Cid < NumComponents; ++Cid)
+		{
+			FVector2f ComponentAlignmentPoint = GetAlignmentPointFromBoundingBoxAndDirection(EdgeDirection, PerComponentBoundingBoxes[Cid]);
+			if (bVertical)
+			{
+				PerComponentTranslation[Cid] = FVector2f(0.0, NextPosition + (OverallAlignmentPoint.Y - ComponentAlignmentPoint.Y));
+			}
+			else
+			{
+				PerComponentTranslation[Cid] = FVector2f(NextPosition + (OverallAlignmentPoint.X - ComponentAlignmentPoint.X), 0.0);
+			}
+			if (bEqualizeSpacing)
+			{
+				float ComponentSpace = bVertical ? PerComponentBoundingBoxes[Cid].Height() : PerComponentBoundingBoxes[Cid].Width();
+				NextPosition += SpreadDirection * (ComponentSpace + GapSpace);
+			}
+			else
+			{
+				NextPosition += SpreadDirection * PerComponentSpace;
+			}
+		}
+		return PerComponentTranslation;
+	};
+
+	switch (DistributeMode)
+	{
+	case EUVEditorDistributeModeBackend::TopEdges:
+		PerComponentTranslation = ComputeDistributeTranslations(true, EUVEditorAlignDirectionBackend::Top, -1.0f, false);
+		break;
+	case EUVEditorDistributeModeBackend::BottomEdges:
+		PerComponentTranslation = ComputeDistributeTranslations(true, EUVEditorAlignDirectionBackend::Bottom, 1.0f, false);
+		break;
+	case EUVEditorDistributeModeBackend::LeftEdges:
+		PerComponentTranslation = ComputeDistributeTranslations(false, EUVEditorAlignDirectionBackend::Left, 1.0f, false);
+		break;
+	case EUVEditorDistributeModeBackend::RightEdges:
+		PerComponentTranslation = ComputeDistributeTranslations(false, EUVEditorAlignDirectionBackend::Right, -1.0f, false);
+		break;
+	case EUVEditorDistributeModeBackend::CentersVertically:
+		PerComponentTranslation = ComputeDistributeTranslations(false, EUVEditorAlignDirectionBackend::CenterVertically, 1.0f, false);
+		break;
+	case EUVEditorDistributeModeBackend::CentersHorizontally:
+		PerComponentTranslation = ComputeDistributeTranslations(true, EUVEditorAlignDirectionBackend::CenterHorizontally, 1.0f, false);
+		break;
+	case EUVEditorDistributeModeBackend::HorizontalSpace:
+		PerComponentTranslation = ComputeDistributeTranslations(false, EUVEditorAlignDirectionBackend::Left, 1.0f, true);
+		break;
+	case EUVEditorDistributeModeBackend::VerticalSpace:
+		PerComponentTranslation = ComputeDistributeTranslations(true, EUVEditorAlignDirectionBackend::Bottom, 1.0f, true);
+		break;
+	default:
+		ensure(false);
+	}
+
+	auto TranslateFunc = [&](const FVector2f& Translation, int32 Vid)
+	{
+		FVector2f UV = FUVEditorUXSettings::InternalUVToExternalUV(ActiveUVLayer->GetElement(Vid));
+		UV = (UV + Translation);
+		ActiveUVLayer->SetElement(Vid, FUVEditorUXSettings::ExternalUVToInternalUV(UV));
+	};
+
+	if (TransformingElements.IsSet())
+	{
+		for (int ElementID : TransformingElements.GetValue())
+		{
+			TranslateFunc(PerComponentTranslation[*ElementToComponent.Find(ElementID)], ElementID);
+			if (Progress && Progress->Cancelled())
+			{
+				return;
+			}
+		}
+	}
+	else
+	{
+		for (int ElementID : ActiveUVLayer->ElementIndicesItr())
+		{
+			TranslateFunc(PerComponentTranslation[*ElementToComponent.Find(ElementID)], ElementID);
+			if (Progress && Progress->Cancelled())
+			{
+				return;
+			}
+		}
+	}
+
+}
+
+
+TUniquePtr<FDynamicMeshOperator> UUVEditorUVTransformOperatorFactory::MakeNewOperator()
+{
+	switch (Settings->TransformType)
+	{
+	case EUVEditorUVTransformType::Transform:
+	{
+		TUniquePtr<FUVEditorUVTransformOp> Op = MakeUnique<FUVEditorUVTransformOp>();
+		Op->OriginalMesh = OriginalMesh;
+		Op->SetTransform(TargetTransform);
+		Op->Selection = Selection;
+		Op->UVLayerIndex = GetSelectedUVChannel();
+
+		Op->Scale = Settings->Scale;
+		Op->Rotation = Settings->Rotation;
+		Op->Translation = Settings->Translation;
+
+		switch (Settings->TranslationMode)
+		{
+		case EUVEditorTranslationMode::Relative:
+			Op->TranslationMode = EUVEditorTranslationModeBackend::Relative;
+			break;
+		case EUVEditorTranslationMode::Absolute:
+			Op->TranslationMode = EUVEditorTranslationModeBackend::Absolute;
+			break;
+		}
+
+		switch (Settings->PivotMode)
+		{
+		case EUVEditorPivotType::Origin:
+			Op->PivotMode = EUVEditorPivotTypeBackend::Origin;
+			break;
+		case EUVEditorPivotType::IndividualBoundingBoxCenter:
+			Op->PivotMode = EUVEditorPivotTypeBackend::IndividualBoundingBoxCenter;
+			break;
+		case EUVEditorPivotType::BoundingBoxCenter:
+			Op->PivotMode = EUVEditorPivotTypeBackend::BoundingBoxCenter;
+			break;
+		case EUVEditorPivotType::Manual:
+			Op->PivotMode = EUVEditorPivotTypeBackend::Manual;
+			break;
+		}
+		Op->ManualPivot = Settings->ManualPivot;
+
+
+		return Op;
+	}
+	break;
+	case EUVEditorUVTransformType::Align:
+	{
+		TUniquePtr<FUVEditorUVAlignOp> Op = MakeUnique<FUVEditorUVAlignOp>();
+		Op->OriginalMesh = OriginalMesh;
+		Op->SetTransform(TargetTransform);
+		Op->Selection = Selection;
+		Op->UVLayerIndex = GetSelectedUVChannel();
+
+		switch (Settings->AlignAnchor)
+		{
+		//case EUVEditorAlignAnchor::FirstItem:
+		//	Op->AlignAnchor = EUVEditorAlignAnchorBackend::FirstItem;
+		//	break;
+		case EUVEditorAlignAnchor::UDIMTile:
+			Op->AlignAnchor = EUVEditorAlignAnchorBackend::UDIMTile;
+			break;
+		case EUVEditorAlignAnchor::BoundingBox:
+			Op->AlignAnchor = EUVEditorAlignAnchorBackend::BoundingBox;
+			break;
+		case EUVEditorAlignAnchor::Manual:
+			Op->AlignAnchor = EUVEditorAlignAnchorBackend::Manual;
+			break;
+		}
+
+		switch (Settings->AlignDirection)
+		{
+		case EUVEditorAlignDirection::Top:
+			Op->AlignDirection = EUVEditorAlignDirectionBackend::Top;
+			break;
+		case EUVEditorAlignDirection::Bottom:
+			Op->AlignDirection = EUVEditorAlignDirectionBackend::Bottom;
+			break;
+		case EUVEditorAlignDirection::Left:
+			Op->AlignDirection = EUVEditorAlignDirectionBackend::Left;
+			break;
+		case EUVEditorAlignDirection::Right:
+			Op->AlignDirection = EUVEditorAlignDirectionBackend::Right;
+			break;
+		case EUVEditorAlignDirection::CenterVertically:
+			Op->AlignDirection = EUVEditorAlignDirectionBackend::CenterVertically;
+			break;
+		case EUVEditorAlignDirection::CenterHorizontally:
+			Op->AlignDirection = EUVEditorAlignDirectionBackend::CenterHorizontally;
+			break;
+		}
+
+		Op->ManualAnchor = Settings->ManualAnchor;
+
+		return Op;
+	}
+	break;
+	case EUVEditorUVTransformType::Distribute:
+	{
+		TUniquePtr<FUVEditorUVDistributeOp> Op = MakeUnique<FUVEditorUVDistributeOp>();
+		Op->OriginalMesh = OriginalMesh;
+		Op->SetTransform(TargetTransform);
+		Op->Selection = Selection;
+		Op->UVLayerIndex = GetSelectedUVChannel();
+
+		switch (Settings->DistributeMode)
+		{
+		case EUVEditorDistributeMode::LeftEdges:
+			Op->DistributeMode = EUVEditorDistributeModeBackend::LeftEdges;
+			break;
+		case EUVEditorDistributeMode::RightEdges:
+			Op->DistributeMode = EUVEditorDistributeModeBackend::RightEdges;
+			break;
+		case EUVEditorDistributeMode::TopEdges:
+			Op->DistributeMode = EUVEditorDistributeModeBackend::TopEdges;
+			break;
+		case EUVEditorDistributeMode::BottomEdges:
+			Op->DistributeMode = EUVEditorDistributeModeBackend::BottomEdges;
+			break;
+		case EUVEditorDistributeMode::CentersVertically:
+			Op->DistributeMode = EUVEditorDistributeModeBackend::CentersVertically;
+			break;
+		case EUVEditorDistributeMode::CentersHorizontally:
+			Op->DistributeMode = EUVEditorDistributeModeBackend::CentersHorizontally;
+			break;
+		case EUVEditorDistributeMode::VerticalSpace:
+			Op->DistributeMode = EUVEditorDistributeModeBackend::VerticalSpace;
+			break;
+		case EUVEditorDistributeMode::HorizontalSpace:
+			Op->DistributeMode = EUVEditorDistributeModeBackend::HorizontalSpace;
+			break;
+		}
+
+		return Op;
+	}
+	break;
+	default:
+		return nullptr;
+	}
+}
diff --git a/Engine/Plugins/Experimental/UVEditor/Source/UVEditorTools/Private/UVEditorTransformTool.cpp b/Engine/Plugins/Experimental/UVEditor/Source/UVEditorTools/Private/UVEditorTransformTool.cpp
new file mode 100644
--- /dev/null
+++ b/Engine/Plugins/Experimental/UVEditor/Source/UVEditorTools/Private/UVEditorTransformTool.cpp
@@ -0,0 +1,285 @@
+// Copyright Epic Games, Inc. All Rights Reserved.
+
+#include "UVEditorTransformTool.h"
+
+#include "ContextObjectStore.h"
+#include "DynamicMesh/DynamicMesh3.h"
+#include "DynamicMesh/DynamicMeshChangeTracker.h"
+#include "InteractiveToolManager.h"
+#include "MeshOpPreviewHelpers.h" // UMeshOpPreviewWithBackgroundCompute
+#include "Operators/UVEditorUVTransformOp.h"
+#include "ToolTargets/UVEditorToolMeshInput.h"
+#include "ContextObjects/UVToolContextObjects.h"
+#include "EngineAnalytics.h"
+#include "UVEditorToolAnalyticsUtils.h"
+#include "UVEditorUXSettings.h"
+#include "CanvasTypes.h"
+#include "Engine/Canvas.h"
+
+using namespace UE::Geometry;
+
+#define LOCTEXT_NAMESPACE "UUVEditorTransformTool"
+
+// Tool builder
+// TODO: Could consider sharing some of the tool builder boilerplate for UV editor tools in a common base class.
+
+bool UUVEditorBaseTransformToolBuilder::CanBuildTool(const FToolBuilderState& SceneState) const
+{
+	return Targets && Targets->Num() > 0;
+}
+
+UInteractiveTool* UUVEditorBaseTransformToolBuilder::BuildTool(const FToolBuilderState& SceneState) const
+{
+	UUVEditorTransformTool* NewTool = NewObject<UUVEditorTransformTool>(SceneState.ToolManager);
+	ConfigureTool(NewTool);	
+	return NewTool;
+}
+
+void UUVEditorBaseTransformToolBuilder::ConfigureTool(UUVEditorTransformTool* NewTool) const
+{
+	NewTool->SetTargets(*Targets);
+}
+
+void UUVEditorTransformToolBuilder::ConfigureTool(UUVEditorTransformTool* NewTool) const
+{
+	UUVEditorBaseTransformToolBuilder::ConfigureTool(NewTool);
+	NewTool->SetToolMode(EUVEditorUVTransformType::Transform);
+}
+
+void UUVEditorAlignToolBuilder::ConfigureTool(UUVEditorTransformTool* NewTool) const
+{
+	UUVEditorBaseTransformToolBuilder::ConfigureTool(NewTool);
+	NewTool->SetToolMode(EUVEditorUVTransformType::Align);
+}
+
+void UUVEditorDistributeToolBuilder::ConfigureTool(UUVEditorTransformTool* NewTool) const
+{
+	UUVEditorBaseTransformToolBuilder::ConfigureTool(NewTool);
+	NewTool->SetToolMode(EUVEditorUVTransformType::Distribute);
+}
+
+void UUVEditorTransformTool::SetToolMode(const EUVEditorUVTransformType& Mode)
+{
+	ToolMode = Mode;
+}
+
+void UUVEditorTransformTool::Setup()
+{
+	check(Targets.Num() > 0);
+
+	ToolStartTimeAnalytics = FDateTime::UtcNow();
+
+	UInteractiveTool::Setup();
+
+	Settings = NewObject<UUVEditorUVTransformProperties>(this);
+	Settings->RestoreProperties(this);
+	Settings->TransformType = ToolMode.Get(EUVEditorUVTransformType::Transform);
+	//Settings->bUDIMCVAREnabled = (FUVEditorUXSettings::CVarEnablePrototypeUDIMSupport.GetValueOnGameThread() > 0);
+	AddToolPropertySource(Settings);
+
+	UContextObjectStore* ContextStore = GetToolManager()->GetContextObjectStore();
+	UVToolSelectionAPI = ContextStore->FindContext<UUVToolSelectionAPI>();
+
+	UUVToolSelectionAPI::FHighlightOptions HighlightOptions;
+	HighlightOptions.bBaseHighlightOnPreviews = true;
+	HighlightOptions.bAutoUpdateUnwrap = true;
+	UVToolSelectionAPI->SetSelectionMechanicMode(UUVToolSelectionAPI::EUVEditorSelectionMode::Triangle);
+	UVToolSelectionAPI->SetHighlightOptions(HighlightOptions);
+	UVToolSelectionAPI->SetHighlightVisible(true, false, true);
+
+	auto SetupOpFactory = [this](UUVEditorToolMeshInput& Target, TSet<int32>* Selection)
+	{
+		TObjectPtr<UUVEditorUVTransformOperatorFactory> Factory = NewObject<UUVEditorUVTransformOperatorFactory>();
+		Factory->TargetTransform = Target.AppliedPreview->PreviewMesh->GetTransform();
+		Factory->Settings = Settings;
+		Factory->OriginalMesh = Target.AppliedCanonical;
+		Factory->GetSelectedUVChannel = [&Target]() { return Target.UVLayerIndex; };
+		if (Selection)
+		{
+			Factory->Selection.Emplace(*Selection);
+		}
+
+		Target.AppliedPreview->ChangeOpFactory(Factory);
+		Target.AppliedPreview->OnMeshUpdated.AddWeakLambda(this, [this, &Target](UMeshOpPreviewWithBackgroundCompute* Preview) {
+			Target.UpdateUnwrapPreviewFromAppliedPreview();
+
+			this->UVToolSelectionAPI->RebuildUnwrapHighlight(Preview->PreviewMesh->GetTransform());
+			});
+
+		Target.AppliedPreview->InvalidateResult();
+		return Factory;
+	};
+
+	if (UVToolSelectionAPI->HaveSelections())
+	{
+		Factories.Reserve(UVToolSelectionAPI->GetSelections().Num());
+		for (FUVToolSelection Selection : UVToolSelectionAPI->GetSelections())
+		{
+			if (ensure(Selection.Type == FUVToolSelection::EType::Triangle))
+			{
+				Factories.Add(SetupOpFactory(*Selection.Target, &Selection.SelectedIDs));
+			}
+		}
+	}
+	else
+	{
+		Factories.Reserve(Targets.Num());
+		for (int32 TargetIndex = 0; TargetIndex < Targets.Num(); ++TargetIndex)
+		{
+			Factories.Add(SetupOpFactory(*Targets[TargetIndex], nullptr));
+		}
+	}
+
+	SetToolDisplayName(LOCTEXT("ToolName", "UV Transform"));
+	GetToolManager()->DisplayMessage(LOCTEXT("OnStartUVTransformTool", "Translate, rotate or scale existing UV Charts using various strategies"),
+		EToolMessageLevel::UserNotification);
+
+	// Analytics
+	InputTargetAnalytics = UVEditorAnalytics::CollectTargetAnalytics(Targets);
+}
+
+void UUVEditorTransformTool::Shutdown(EToolShutdownType ShutdownType)
+{
+	Settings->SaveProperties(this);
+	for (TObjectPtr<UUVEditorToolMeshInput> Target : Targets)
+	{
+		Target->AppliedPreview->OnMeshUpdated.RemoveAll(this);
+	}
+
+	if (ShutdownType == EToolShutdownType::Accept)
+	{
+		UUVToolEmitChangeAPI* ChangeAPI = GetToolManager()->GetContextObjectStore()->FindContext<UUVToolEmitChangeAPI>();
+		const FText TransactionName(LOCTEXT("TransformTransactionName", "Transform Tool"));
+		ChangeAPI->BeginUndoTransaction(TransactionName);
+
+		for (TObjectPtr<UUVEditorToolMeshInput> Target : Targets)
+		{
+			// Set things up for undo. 
+			// TODO: It's not entirely clear whether it would be safe to use a FMeshVertexChange instead... It seems like
+			// when bAllowFlips is true, we would end up with changes to the tris of the unwrap. Also, if we stick to saving
+			// all the tris and verts, should we consider using the new dynamic mesh serialization?
+			FDynamicMeshChangeTracker ChangeTracker(Target->UnwrapCanonical.Get());
+			ChangeTracker.BeginChange();
+
+			for (int32 Tid : Target->UnwrapCanonical->TriangleIndicesItr())
+			{
+				ChangeTracker.SaveTriangle(Tid, true);
+			}
+
+			// TODO: Again, it's not clear whether we need to update the entire triangle topology...
+			Target->UpdateCanonicalFromPreviews();
+
+			ChangeAPI->EmitToolIndependentUnwrapCanonicalChange(Target, ChangeTracker.EndChange(), LOCTEXT("ApplyTransformTool", "Transform Tool"));
+		}
+
+		ChangeAPI->EndUndoTransaction();
+
+		// Analytics
+		RecordAnalytics();
+	}
+	else
+	{
+		// Reset the inputs
+		for (TObjectPtr<UUVEditorToolMeshInput> Target : Targets)
+		{
+			Target->UpdatePreviewsFromCanonical();
+		}
+	}
+
+
+	for (TObjectPtr<UUVEditorToolMeshInput> Target : Targets)
+	{
+		Target->AppliedPreview->ClearOpFactory();
+	}
+
+	for (int32 TargetIndex = 0; TargetIndex < Targets.Num(); ++TargetIndex)
+	{
+		Factories[TargetIndex] = nullptr;
+	}
+
+	Settings = nullptr;
+	Targets.Empty();
+}
+
+void UUVEditorTransformTool::OnTick(float DeltaTime)
+{
+	for (TObjectPtr<UUVEditorToolMeshInput> Target : Targets)
+	{
+		Target->AppliedPreview->Tick(DeltaTime);
+	}
+}
+
+void UUVEditorTransformTool::OnPropertyModified(UObject* PropertySet, FProperty* Property)
+{
+	for (TObjectPtr<UUVEditorToolMeshInput> Target : Targets)
+	{
+		Target->AppliedPreview->InvalidateResult();
+	}
+}
+
+bool UUVEditorTransformTool::CanAccept() const
+{
+	for (TObjectPtr<UUVEditorToolMeshInput> Target : Targets)
+	{
+		if (!Target->AppliedPreview->HaveValidResult())
+		{
+			return false;
+		}
+	}
+	return true;
+}
+
+void UUVEditorTransformTool::DrawHUD(FCanvas* Canvas, IToolsContextRenderAPI* RenderAPI)
+{
+	// TODO: Add support here for highlighting first selected item for alignment visualization
+}
+
+
+void UUVEditorTransformTool::RecordAnalytics()
+{
+	using namespace UVEditorAnalytics;
+
+	if (!FEngineAnalytics::IsAvailable())
+	{
+		return;
+	}
+
+	TArray<FAnalyticsEventAttribute> Attributes;
+	Attributes.Add(FAnalyticsEventAttribute(TEXT("Timestamp"), FDateTime::UtcNow().ToString()));
+
+	// Tool inputs
+	InputTargetAnalytics.AppendToAttributes(Attributes, "Input");
+
+	// Tool stats
+	if (CanAccept())
+	{
+		TArray<double> PerAssetValidResultComputeTimes;
+		for (TObjectPtr<UUVEditorToolMeshInput> Target : Targets)
+		{
+			// Note: This would log -1 if the result was invalid, but checking CanAccept above ensures results are valid
+			PerAssetValidResultComputeTimes.Add(Target->AppliedPreview->GetValidResultComputeTime());
+		}
+		Attributes.Add(FAnalyticsEventAttribute(TEXT("Stats.PerAsset.ComputeTimeSeconds"), PerAssetValidResultComputeTimes));
+	}
+	Attributes.Add(FAnalyticsEventAttribute(TEXT("Stats.ToolActiveDuration"), (FDateTime::UtcNow() - ToolStartTimeAnalytics).ToString()));
+
+	// Tool settings chosen by the user
+	//Attributes.Add(AnalyticsEventAttributeEnum(TEXT("Settings.LayoutType"), Settings->LayoutType));
+	//Attributes.Add(FAnalyticsEventAttribute(TEXT("Settings.TextureResolution"), Settings->TextureResolution));
+	//Attributes.Add(FAnalyticsEventAttribute(TEXT("Settings.Scale"), Settings->Scale));
+	//const TArray<FVector2D::FReal> TranslationArray({ Settings->Translation.X, Settings->Translation.Y });
+	//Attributes.Add(FAnalyticsEventAttribute(TEXT("Settings.Translation"), TranslationArray));
+	//Attributes.Add(FAnalyticsEventAttribute(TEXT("Settings.AllowFlips"), Settings->bAllowFlips));
+
+	FEngineAnalytics::GetProvider().RecordEvent(UVEditorAnalyticsEventName(TEXT("TransformTool")), Attributes);
+
+	if constexpr (false)
+	{
+		for (const FAnalyticsEventAttribute& Attr : Attributes)
+		{
+			UE_LOG(LogGeometry, Log, TEXT("Debug %s.TransformTool.%s = %s"), *UVEditorAnalyticsPrefix, *Attr.GetName(), *Attr.GetValue());
+		}
+	}
+}
+
+#undef LOCTEXT_NAMESPACE
diff --git a/Engine/Plugins/Experimental/UVEditor/Source/UVEditorTools/Private/Utilities/MeshUDIMClassifier.cpp b/Engine/Plugins/Experimental/UVEditor/Source/UVEditorTools/Private/Utilities/MeshUDIMClassifier.cpp
--- a/Engine/Plugins/Experimental/UVEditor/Source/UVEditorTools/Private/Utilities/MeshUDIMClassifier.cpp
+++ b/Engine/Plugins/Experimental/UVEditor/Source/UVEditorTools/Private/Utilities/MeshUDIMClassifier.cpp
@@ -33,6 +33,40 @@ TArray<int32> FDynamicMeshUDIMClassifier::TidsForTile(FVector2i TileIndexIn) con
 	}
 }
 
+FVector2i FDynamicMeshUDIMClassifier::ClassifyTrianglesToUDIM(const FDynamicMeshUVOverlay* UVOverlay, TArray<int32> Tids)
+{
+	auto UVTriangleToUDIM = [UVOverlay](int32 Tid)
+	{
+		FVector2i UDIM;
+		FVector2f Vertex0, Vertex1, Vertex2;
+		UVOverlay->GetTriElements(Tid, Vertex0, Vertex1, Vertex2);
+		FVector2f BaryCenter = (Vertex0 + Vertex1 + Vertex2) / 3.0f;
+		BaryCenter = FUVEditorUXSettings::InternalUVToExternalUV(BaryCenter);
+		UDIM.X = FMath::Floor(BaryCenter.X);
+		UDIM.Y = FMath::Floor(BaryCenter.Y);
+		return UDIM;
+	};
+
+	TMap<FVector2i, int32> UDIMCount;
+	for (int32 Tid : Tids)
+	{
+		FVector2i UDIM = UVTriangleToUDIM(Tid);
+		UDIMCount.FindOrAdd(UDIM, 0)++;
+	}
+	FVector2i MaximumUDIM;
+	int32 MaximumUDIMCount = 0;
+	for (auto UDIMEntry : UDIMCount)
+	{
+		if (UDIMEntry.Value > MaximumUDIMCount)
+		{
+			MaximumUDIMCount = UDIMEntry.Value;
+			MaximumUDIM = UDIMEntry.Key;
+		}
+	}
+	return MaximumUDIM;
+}
+
+
 void FDynamicMeshUDIMClassifier::ClassifyUDIMs()
 {
 	int32 CurrentUDIMIndex = 0;
@@ -53,42 +87,9 @@ void FDynamicMeshUDIMClassifier::ClassifyUDIMs()
 		UVComponents.FindConnectedTriangles(UVIslandPredicate);
 	}
 
-	auto UVTriangleToUDIM = [this](int32 Tid)
-	{
-		FVector2i UDIM;
-		FVector2f Vertex0, Vertex1, Vertex2;
-		UVOverlay->GetTriElements(Tid, Vertex0, Vertex1, Vertex2);
-		FVector2f BaryCenter = (Vertex0 + Vertex1 + Vertex2) / 3.0f;
-		BaryCenter = FUVEditorUXSettings::InternalUVToExternalUV(BaryCenter);
-		UDIM.X = FMath::Floor(BaryCenter.X);
-		UDIM.Y = FMath::Floor(BaryCenter.Y);
-		return UDIM;
-	};
-
-	auto IslandToUDIM = [this, &UVComponents, &UVTriangleToUDIM](int32 Cid)
-	{
-		TMap<FVector2i, int32> UDIMCount;
-		for (int32 Tid : UVComponents.GetComponent(Cid).Indices)
-		{
-			FVector2i UDIM = UVTriangleToUDIM(Tid);
-			UDIMCount.FindOrAdd(UDIM, 0)++;
-		}
-		FVector2i MaximumUDIM;
-		int32 MaximumUDIMCount = 0;
-		for (auto UDIMEntry : UDIMCount)
-		{
-			if (UDIMEntry.Value > MaximumUDIMCount)
-			{
-				MaximumUDIMCount = UDIMEntry.Value;
-				MaximumUDIM = UDIMEntry.Key;
-			}
-		}
-		return MaximumUDIM;
-	};
-
 	for (int32 Cid = 0; Cid < UVComponents.Num(); ++Cid)
 	{
-		FVector2i UDIM = IslandToUDIM(Cid);
+		FVector2i UDIM = ClassifyTrianglesToUDIM(UVOverlay, UVComponents.GetComponent(Cid).Indices);
 		UDIMs.FindOrAdd(UDIM).Append(UVComponents.GetComponent(Cid).Indices);
 	}
 
diff --git a/Engine/Plugins/Experimental/UVEditor/Source/UVEditorTools/Public/Operators/UVEditorUVTransformOp.h b/Engine/Plugins/Experimental/UVEditor/Source/UVEditorTools/Public/Operators/UVEditorUVTransformOp.h
new file mode 100644
--- /dev/null
+++ b/Engine/Plugins/Experimental/UVEditor/Source/UVEditorTools/Public/Operators/UVEditorUVTransformOp.h
@@ -0,0 +1,328 @@
+// Copyright Epic Games, Inc. All Rights Reserved.
+
+#pragma once
+
+#include "CoreMinimal.h"
+#include "Util/ProgressCancel.h"
+#include "ModelingOperators.h"
+#include "InteractiveTool.h"
+#include "DynamicMesh/DynamicMeshAttributeSet.h"
+
+#include "UVEditorUVTransformOp.generated.h"
+
+class UUVTransformProperties;
+
+/**
+ * UV Transform Strategies for the UV Transform Tool
+ */
+UENUM()
+enum class EUVEditorUVTransformType
+{
+	/** Apply Scale, Totation and Translation properties to all UV values */
+	Transform,
+	/** Position all selected elements given alignment rules */
+	Align,
+	/** Position all selected elements given distribution rules */
+	Distribute
+};
+
+/**
+ * Transform Pivot Mode
+ */
+UENUM()
+enum class EUVEditorPivotType
+{
+	/** Pivot around the collective center of all island bounding boxes */
+	BoundingBoxCenter,
+	/** Pivot around the global origin point */
+	Origin,
+	/** Pivot around each island's bounding box center */
+	IndividualBoundingBoxCenter,
+	/** Pivot around a user specified point */
+	Manual
+};
+
+/**
+ * Translation Mode
+ */
+UENUM()
+enum class EUVEditorTranslationMode
+{
+	/** Move elements relative to their current position by the amount specified */
+	Relative,
+	/** Move elements such that the transform origin is placed at the value specified */
+	Absolute
+};
+
+UENUM()
+enum class EUVEditorAlignDirection
+{
+	Top,
+	Bottom,
+	Left,
+	Right,
+	CenterVertically,
+	CenterHorizontally
+};
+
+UENUM()
+enum class EUVEditorAlignAnchor
+{
+	//FirstItem, // TODO: Support this later once we have support in the Selection API
+
+	/** Align relative to the collective bounding box of all islands */
+	BoundingBox,
+	/** Align relative to the local UDIM tile containing the island */
+	UDIMTile,
+	/** Align relative to a user specified point */
+	Manual
+};
+
+UENUM()
+enum class EUVEditorDistributeMode
+{
+	VerticalSpace,
+	HorizontalSpace,
+	TopEdges,
+	BottomEdges,
+	LeftEdges,
+	RightEdges,
+	CentersVertically,
+	CentersHorizontally
+};
+
+
+/**
+ * UV Transform Settings
+ */
+UCLASS()
+class UVEDITORTOOLS_API UUVEditorUVTransformProperties : public UInteractiveToolPropertySet
+{
+	GENERATED_BODY()
+
+public:
+	/** Type of Transform applied to input UVs */
+	UPROPERTY()
+	EUVEditorUVTransformType TransformType = EUVEditorUVTransformType::Transform;
+
+	/** Scale applied to UVs, potentially non-uniform */
+	UPROPERTY(EditAnywhere, Category = "Scaling", Transient, meta = (DisplayName = "Scale",
+		      EditCondition = "TransformType == EUVEditorUVTransformType::Transform", EditConditionHides, HideEditConditionToggle = true))
+	FVector2D Scale = FVector2D(1.0, 1.0); 
+
+	/** Rotation applied to UVs after scaling, specified in degrees */
+	UPROPERTY(EditAnywhere, Category = "Rotation", Transient, meta = (DisplayName = "Rotation",
+		      EditCondition = "TransformType == EUVEditorUVTransformType::Transform", EditConditionHides, HideEditConditionToggle = true))
+	float Rotation = 0.0;
+
+	/** Translation applied to UVs, and after scaling and rotation */
+	UPROPERTY(EditAnywhere, Category = "Translation", Transient, meta = (DisplayName = "Transform",
+		      EditCondition = "TransformType == EUVEditorUVTransformType::Transform", EditConditionHides, HideEditConditionToggle = true))
+	FVector2D Translation = FVector2D(0, 0);
+
+	/** Translation applied to UVs, and after scaling and rotation */
+	UPROPERTY(EditAnywhere, Category = "Translation", Transient, meta = (DisplayName = "Transform Mode",
+		EditCondition = "TransformType == EUVEditorUVTransformType::Transform", EditConditionHides, HideEditConditionToggle = true))
+	EUVEditorTranslationMode TranslationMode = EUVEditorTranslationMode::Relative;
+
+	/** Transformation origin mode used for scaling and rotation */
+	UPROPERTY(EditAnywhere, Category = "Transform Origin", Transient, meta = (DisplayName = "Mode",
+		EditCondition = "TransformType == EUVEditorUVTransformType::Transform", EditConditionHides, HideEditConditionToggle = true))
+	EUVEditorPivotType PivotMode = EUVEditorPivotType::BoundingBoxCenter;
+
+	/** Manual Transformation origin point */
+	UPROPERTY(EditAnywhere, Category = "Transform Origin", Transient, meta = (DisplayName = "Coords",
+		EditCondition = "TransformType == EUVEditorUVTransformType::Transform && PivotMode == EUVEditorPivotType::Manual", EditConditionHides, HideEditConditionToggle = true))
+	FVector2D ManualPivot = FVector2D(0, 0);
+
+	/** Controls what geometry the alignment is to be relative to when performed. */
+	UPROPERTY(EditAnywhere, Category = "Align", meta = (DisplayName = "Alignment Anchor",
+		EditCondition = "TransformType == EUVEditorUVTransformType::Align", EditConditionHides, HideEditConditionToggle = true))
+	EUVEditorAlignAnchor AlignAnchor = EUVEditorAlignAnchor::BoundingBox;
+
+	/** Manual anchor location for relative alignment */
+	UPROPERTY(EditAnywhere, Category = "Align", meta = (DisplayName = "Anchor Coords",
+		EditCondition = "TransformType == EUVEditorUVTransformType::Align && AlignAnchor == EUVEditorAlignAnchor::Manual", EditConditionHides, HideEditConditionToggle = true))
+	FVector2D ManualAnchor = FVector2D(0, 0);
+
+	/** Controls what side of the island bounding boxes are being aligned */
+	UPROPERTY(EditAnywhere, Category = "Align", meta = (DisplayName = "Alignment Direction",
+		EditCondition = "TransformType == EUVEditorUVTransformType::Align", EditConditionHides, HideEditConditionToggle = true))
+	EUVEditorAlignDirection AlignDirection = EUVEditorAlignDirection::Top;
+
+	/** Controls the distribution behavior */
+	UPROPERTY(EditAnywhere, Category = "Distribute", meta = (DisplayName = "Distribution Mode",
+		EditCondition = "TransformType == EUVEditorUVTransformType::Distribute", EditConditionHides, HideEditConditionToggle = true))
+	EUVEditorDistributeMode DistributeMode = EUVEditorDistributeMode::TopEdges;
+
+};
+
+
+namespace UE
+{
+namespace Geometry
+{
+	class FDynamicMesh3;
+	class FDynamicMeshUVPacker;
+	class FMeshConnectedComponents;
+
+	enum class EUVEditorUVTransformTypeBackend
+	{
+		Transform,
+		Align,
+		Distribute
+	};
+
+	enum class EUVEditorPivotTypeBackend
+	{
+		Origin,
+		IndividualBoundingBoxCenter,
+		BoundingBoxCenter,
+		Manual
+	};
+
+	UENUM()
+	enum class EUVEditorTranslationModeBackend
+	{
+		Relative,
+		Absolute
+	};
+
+	enum class EUVEditorAlignDirectionBackend
+	{
+		Top,
+		Bottom,
+		Left,
+		Right,
+		CenterVertically,
+		CenterHorizontally
+	};
+
+	enum class EUVEditorAlignAnchorBackend
+	{
+		//FirstItem, // TODO:  Support this later once we have support in the Selection API
+		UDIMTile,
+		BoundingBox,
+		Manual
+	};
+
+	enum class EUVEditorDistributeModeBackend
+	{
+		LeftEdges,
+		RightEdges,
+		TopEdges,
+		BottomEdges,
+		CentersVertically,
+		CentersHorizontally,
+		VerticalSpace,
+		HorizontalSpace
+	};
+
+class UVEDITORTOOLS_API FUVEditorUVTransformBaseOp : public FDynamicMeshOperator
+{
+public:
+	virtual ~FUVEditorUVTransformBaseOp() {}
+
+	// inputs
+	TSharedPtr<FDynamicMesh3, ESPMode::ThreadSafe> OriginalMesh;
+	TOptional<TSet<int32>> Selection;
+	int UVLayerIndex = 0;
+
+	void SetTransform(const FTransformSRT3d& Transform);
+
+	virtual void CalculateResult(FProgressCancel* Progress) override;
+
+protected:
+	virtual void HandleTransformationOp(FProgressCancel* Progress) = 0;
+
+	FVector2f GetAlignmentPointFromBoundingBoxAndDirection(EUVEditorAlignDirectionBackend Direction, const FAxisAlignedBox2d& BoundingBox);
+	FVector2f GetAlignmentPointFromUDIMAndDirection(EUVEditorAlignDirectionBackend Direction, FVector2i UDIMTile);
+	void RebuildBoundingBoxes();
+	void CollectTransformElements();
+
+	TMap<int32, int32> ElementToComponent;
+	FDynamicMeshUVOverlay* ActiveUVLayer;
+	FAxisAlignedBox2d OverallBoundingBox;
+	TArray<FAxisAlignedBox2d> PerComponentBoundingBoxes;
+	TOptional< TSet<int32> > TransformingElements;
+	TSharedPtr<FMeshConnectedComponents> UVComponents;
+};
+
+
+/**
+ */
+class UVEDITORTOOLS_API FUVEditorUVTransformOp : public FUVEditorUVTransformBaseOp
+{
+public:
+	virtual ~FUVEditorUVTransformOp() {}
+
+	// inputs
+	FVector2D Scale = FVector2D(1.0, 1.0);
+	float Rotation = 0.0;
+	FVector2D Translation = FVector2D(0, 0);
+	EUVEditorTranslationModeBackend TranslationMode = EUVEditorTranslationModeBackend::Relative;
+
+	EUVEditorPivotTypeBackend PivotMode = EUVEditorPivotTypeBackend::Origin;
+	FVector2D ManualPivot = FVector2D(0, 0);
+
+protected:
+	virtual void HandleTransformationOp(FProgressCancel* Progress) override;
+	FVector2f GetPivotFromMode(int32 ElementID);
+};
+
+class UVEDITORTOOLS_API FUVEditorUVAlignOp : public FUVEditorUVTransformBaseOp
+{
+public:
+	virtual ~FUVEditorUVAlignOp() {}
+
+	EUVEditorAlignAnchorBackend AlignAnchor = EUVEditorAlignAnchorBackend::BoundingBox;
+	EUVEditorAlignDirectionBackend AlignDirection = EUVEditorAlignDirectionBackend::Top;
+	FVector2D ManualAnchor = FVector2D(0, 0);
+
+protected:
+	virtual void HandleTransformationOp(FProgressCancel* Progress) override;
+
+
+
+};
+
+class UVEDITORTOOLS_API FUVEditorUVDistributeOp : public FUVEditorUVTransformBaseOp
+{
+public:
+	virtual ~FUVEditorUVDistributeOp() {}
+
+	EUVEditorDistributeModeBackend DistributeMode = EUVEditorDistributeModeBackend::TopEdges;
+
+protected:
+	virtual void HandleTransformationOp(FProgressCancel* Progress) override;
+
+};
+
+
+} // end namespace UE::Geometry
+} // end namespace UE
+
+/**
+ * Can be hooked up to a UMeshOpPreviewWithBackgroundCompute to perform UV Transform operations.
+ *
+ * Inherits from UObject so that it can hold a strong pointer to the settings UObject, which
+ * needs to be a UObject to be displayed in the details panel.
+ */
+UCLASS()
+class UVEDITORTOOLS_API UUVEditorUVTransformOperatorFactory : public UObject, public UE::Geometry::IDynamicMeshOperatorFactory
+{
+	GENERATED_BODY()
+
+public:
+	// IDynamicMeshOperatorFactory API
+	virtual TUniquePtr<UE::Geometry::FDynamicMeshOperator> MakeNewOperator() override;
+
+	UPROPERTY()
+	TObjectPtr<UUVEditorUVTransformProperties> Settings;
+
+	TOptional<TSet<int32>> Selection;
+
+	TSharedPtr<UE::Geometry::FDynamicMesh3, ESPMode::ThreadSafe> OriginalMesh;
+	TUniqueFunction<int32()> GetSelectedUVChannel = []() { return 0; };
+	FTransform TargetTransform;
+};
\ No newline at end of file
diff --git a/Engine/Plugins/Experimental/UVEditor/Source/UVEditorTools/Public/UVEditorTransformTool.h b/Engine/Plugins/Experimental/UVEditor/Source/UVEditorTools/Public/UVEditorTransformTool.h
new file mode 100644
--- /dev/null
+++ b/Engine/Plugins/Experimental/UVEditor/Source/UVEditorTools/Public/UVEditorTransformTool.h
@@ -0,0 +1,120 @@
+// Copyright Epic Games, Inc. All Rights Reserved.
+
+#pragma once
+
+#include "CoreMinimal.h"
+
+#include "InteractiveTool.h"
+#include "InteractiveToolBuilder.h"
+#include "UVEditorToolAnalyticsUtils.h"
+#include "Selection/UVToolSelectionAPI.h"
+
+#include "UVEditorTransformTool.generated.h"
+
+class UUVEditorToolMeshInput;
+class UUVEditorUVTransformProperties;
+enum class EUVEditorUVTransformType;
+class UUVEditorUVTransformOperatorFactory;
+class UCanvas;
+class UUVEditorTransformTool;
+
+UCLASS()
+class UVEDITORTOOLS_API UUVEditorBaseTransformToolBuilder : public UInteractiveToolBuilder
+{
+	GENERATED_BODY()
+
+public:
+	virtual bool CanBuildTool(const FToolBuilderState& SceneState) const override;
+	virtual UInteractiveTool* BuildTool(const FToolBuilderState& SceneState) const override;
+
+	// This is a pointer so that it can be updated under the builder without
+	// having to set it in the mode after initializing targets.
+	const TArray<TObjectPtr<UUVEditorToolMeshInput>>* Targets = nullptr;
+
+protected:
+	virtual void ConfigureTool(UUVEditorTransformTool* NewTool) const;
+};
+
+UCLASS()
+class UVEDITORTOOLS_API UUVEditorTransformToolBuilder : public UUVEditorBaseTransformToolBuilder
+{
+	GENERATED_BODY()
+
+protected:
+	virtual void ConfigureTool(UUVEditorTransformTool* NewTool) const override;
+};
+
+UCLASS()
+class UVEDITORTOOLS_API UUVEditorAlignToolBuilder : public UUVEditorBaseTransformToolBuilder
+{
+	GENERATED_BODY()
+
+protected:
+	virtual void ConfigureTool(UUVEditorTransformTool* NewTool) const  override;
+};
+
+UCLASS()
+class UVEDITORTOOLS_API UUVEditorDistributeToolBuilder : public UUVEditorBaseTransformToolBuilder
+{
+	GENERATED_BODY()
+
+protected:
+	virtual void ConfigureTool(UUVEditorTransformTool* NewTool) const  override;
+};
+
+/**
+ * 
+ */
+UCLASS()
+class UVEDITORTOOLS_API UUVEditorTransformTool : public UInteractiveTool, public IUVToolSupportsSelection
+{
+	GENERATED_BODY()
+
+public:
+
+	/**
+	 * The tool will operate on the meshes given here.
+	 */
+	virtual void SetTargets(const TArray<TObjectPtr<UUVEditorToolMeshInput>>& TargetsIn)
+	{
+		Targets = TargetsIn;
+	}
+
+	void SetToolMode(const EUVEditorUVTransformType& Mode);
+
+	virtual void Setup() override;
+	virtual void Shutdown(EToolShutdownType ShutdownType) override;
+
+	virtual void OnTick(float DeltaTime) override;
+
+	virtual bool HasCancel() const override { return true; }
+	virtual bool HasAccept() const override { return true; }
+	virtual bool CanAccept() const override;
+
+	virtual void DrawHUD(FCanvas* Canvas, IToolsContextRenderAPI* RenderAPI) override;
+
+	virtual void OnPropertyModified(UObject* PropertySet, FProperty* Property) override;
+protected:	
+
+	TOptional<EUVEditorUVTransformType> ToolMode;
+
+	UPROPERTY()
+	TArray<TObjectPtr<UUVEditorToolMeshInput>> Targets;
+
+	UPROPERTY()
+	TObjectPtr<UUVEditorUVTransformProperties> Settings = nullptr;
+
+	UPROPERTY()
+	TArray<TObjectPtr<UUVEditorUVTransformOperatorFactory>> Factories;
+
+	UPROPERTY()
+	TObjectPtr<UUVToolSelectionAPI> UVToolSelectionAPI = nullptr;
+
+	//
+	// Analytics
+	//
+
+	UE::Geometry::UVEditorAnalytics::FTargetAnalytics InputTargetAnalytics;
+	FDateTime ToolStartTimeAnalytics;
+	void RecordAnalytics();
+};
diff --git a/Engine/Plugins/Experimental/UVEditor/Source/UVEditorTools/Public/Utilities/MeshUDIMClassifier.h b/Engine/Plugins/Experimental/UVEditor/Source/UVEditorTools/Public/Utilities/MeshUDIMClassifier.h
--- a/Engine/Plugins/Experimental/UVEditor/Source/UVEditorTools/Public/Utilities/MeshUDIMClassifier.h
+++ b/Engine/Plugins/Experimental/UVEditor/Source/UVEditorTools/Public/Utilities/MeshUDIMClassifier.h
@@ -25,6 +25,8 @@ namespace UE
 			TArray<FVector2i> ActiveTiles() const;
 			TArray<int32> TidsForTile(FVector2i TileIndexIn) const;
 
+			static FVector2i ClassifyTrianglesToUDIM(const FDynamicMeshUVOverlay* UVOverlay, TArray<int32> Tids);
+
 		protected:
 
 			void ClassifyUDIMs();
