UVEditor: Add ability to get Render and DrawHUD calls form the 3d live preview viewport.

diff --git a/Engine/Plugins/Experimental/UVEditor/Source/UVEditor/Private/UVEditorMode.cpp b/Engine/Plugins/Experimental/UVEditor/Source/UVEditor/Private/UVEditorMode.cpp
--- a/Engine/Plugins/Experimental/UVEditor/Source/UVEditor/Private/UVEditorMode.cpp
+++ b/Engine/Plugins/Experimental/UVEditor/Source/UVEditor/Private/UVEditorMode.cpp
@@ -15,7 +15,7 @@
 #include "Drawing/MeshElementsVisualizer.h"
 #include "Editor.h"
 #include "EditorViewportClient.h"
-#include "EdModeInteractiveToolsContext.h" //ToolsContext
+#include "EdModeInteractiveToolsContext.h" //ToolsContext, EditorInteractiveToolsContext
 #include "EngineAnalytics.h"
 #include "Framework/Commands/UICommandList.h"
 #include "InteractiveTool.h"
@@ -39,6 +39,7 @@
 #include "UVEditorSeamTool.h"
 #include "UVEditorRecomputeUVsTool.h"
 #include "UVSelectTool.h"
+#include "UVEditorInitializationContext.h"
 #include "UVEditorModeToolkit.h"
 #include "UVEditorSubsystem.h"
 #include "ContextObjects/UVToolContextObjects.h"
@@ -269,6 +270,41 @@ void UUVEditorMode::Enter()
 	GetInteractiveToolsContext()->OnRender.AddUObject(this, &UUVEditorMode::Render);
 	GetInteractiveToolsContext()->OnDrawHUD.AddUObject(this, &UUVEditorMode::DrawHUD);
 
+	// We also want the Render and DrawHUD calls from the 3D viewport
+	UContextObjectStore* ContextStore = GetInteractiveToolsContext()->ToolManager->GetContextObjectStore();
+	UUVEditorInitializationContext* InitContext = ContextStore->FindContext<UUVEditorInitializationContext>();
+	UUVToolLivePreviewAPI* LivePreviewAPI = ContextStore->FindContext<UUVToolLivePreviewAPI>();
+	if (ensure(InitContext && InitContext->LivePreviewITC.IsValid() && LivePreviewAPI))
+	{
+		LivePreviewITC = InitContext->LivePreviewITC;
+		LivePreviewITC->OnRender.AddWeakLambda(this, [this, LivePreviewAPI](IToolsContextRenderAPI* RenderAPI) 
+		{ 
+			if (LivePreviewAPI->IsValidLowLevel())
+			{
+				LivePreviewAPI->OnRender.Broadcast(RenderAPI);
+			}
+
+			// This could have been attached to the LivePreviewAPI delegate the way that tools might do it,
+			// but might as well do the call ourselves. Same for DrawHUD.
+			if (SelectionAPI)
+			{
+				SelectionAPI->LivePreviewRender(RenderAPI);
+			}
+		});
+		LivePreviewITC->OnDrawHUD.AddWeakLambda(this, [this, LivePreviewAPI](FCanvas* Canvas, IToolsContextRenderAPI* RenderAPI)
+		{
+			if (LivePreviewAPI->IsValidLowLevel())
+			{
+				LivePreviewAPI->OnDrawHUD.Broadcast(Canvas, RenderAPI);
+			}
+
+			if (SelectionAPI)
+			{
+				SelectionAPI->LivePreviewDrawHUD(Canvas, RenderAPI);
+			}
+		});
+	}
+
 	InitializeModeContexts();
 	InitializeTargets();
 
@@ -688,6 +724,11 @@ void UUVEditorMode::Exit()
 
 	GetInteractiveToolsContext()->OnRender.RemoveAll(this);
 	GetInteractiveToolsContext()->OnDrawHUD.RemoveAll(this);
+	if (LivePreviewITC.IsValid())
+	{
+		LivePreviewITC->OnRender.RemoveAll(this);
+		LivePreviewITC->OnDrawHUD.RemoveAll(this);
+	}
 
 	Super::Exit();
 }
@@ -720,6 +761,14 @@ void UUVEditorMode::InitializeAssetEditorContexts(UContextObjectStore& ContextSt
 		ContextStore.AddContextObject(LivePreviewAPI);
 	}
 
+	// Prep the editor-only context that we use to pass things to the mode.
+	if (!ContextStore.FindContext<UUVEditorInitializationContext>())
+	{
+		UUVEditorInitializationContext* InitContext = NewObject<UUVEditorInitializationContext>();
+		InitContext->LivePreviewITC = Cast<UEditorInteractiveToolsContext>(LivePreviewModeManager.GetInteractiveToolsContext());
+		ContextStore.AddContextObject(InitContext);
+	}
+
 	if (!ContextStore.FindContext<UUVToolViewportButtonsAPI>())
 	{
 		ContextStore.AddContextObject(&ViewportButtonsAPI);
diff --git a/Engine/Plugins/Experimental/UVEditor/Source/UVEditor/Public/UVEditorInitializationContext.h b/Engine/Plugins/Experimental/UVEditor/Source/UVEditor/Public/UVEditorInitializationContext.h
new file mode 100644
--- /dev/null
+++ b/Engine/Plugins/Experimental/UVEditor/Source/UVEditor/Public/UVEditorInitializationContext.h
@@ -0,0 +1,27 @@
+// Copyright Epic Games, Inc. All Rights Reserved.
+
+#pragma once
+
+#include "CoreMinimal.h"
+
+#include "ContextObjects/UVToolContextObjects.h"
+
+#include "UVEditorInitializationContext.generated.h"
+
+class UEditorInteractiveToolsContext;
+
+/**
+ * An editor-only context used specifically to pass data from the asset editor to the
+ * mode when that data is not meant to be used by the tools. This allows a mode
+ * to be prepped before its first Enter() call even when it doesn't exist yet.
+ * 
+ * Data/api's that are meant to be used by the tools would usually be placed into
+ * other, non-editor-only context objects (UUVToolLivePreviewAPI, etc).
+ */
+UCLASS()
+class UVEDITOR_API UUVEditorInitializationContext : public UUVToolContextObject
+{
+	GENERATED_BODY()
+public:
+	TWeakObjectPtr<UEditorInteractiveToolsContext> LivePreviewITC;
+};
\ No newline at end of file
diff --git a/Engine/Plugins/Experimental/UVEditor/Source/UVEditor/Public/UVEditorMode.h b/Engine/Plugins/Experimental/UVEditor/Source/UVEditor/Public/UVEditorMode.h
--- a/Engine/Plugins/Experimental/UVEditor/Source/UVEditor/Public/UVEditorMode.h
+++ b/Engine/Plugins/Experimental/UVEditor/Source/UVEditor/Public/UVEditorMode.h
@@ -14,6 +14,7 @@
 PREDECLARE_GEOMETRY(class FDynamicMesh3);
 
 class FAssetEditorModeManager;
+class UEditorInteractiveToolsContext;
 class FEditorViewportClient;
 class FToolCommandChange;
 class UContextObjectStore;
@@ -331,6 +332,9 @@ protected:
 	UPROPERTY()
 	TObjectPtr<UWorld> LivePreviewWorld = nullptr;
 
+	// Used to forward Render/DrawHUD calls in the live preview to the api object.
+	TWeakObjectPtr<UEditorInteractiveToolsContext> LivePreviewITC;
+
 	UPROPERTY()
 	TObjectPtr<UUVToolSelectionAPI> SelectionAPI = nullptr;
 
diff --git a/Engine/Plugins/Experimental/UVEditor/Source/UVEditorTools/Private/ContextObjects/UVToolContextObjects.cpp b/Engine/Plugins/Experimental/UVEditor/Source/UVEditorTools/Private/ContextObjects/UVToolContextObjects.cpp
--- a/Engine/Plugins/Experimental/UVEditor/Source/UVEditorTools/Private/ContextObjects/UVToolContextObjects.cpp
+++ b/Engine/Plugins/Experimental/UVEditor/Source/UVEditorTools/Private/ContextObjects/UVToolContextObjects.cpp
@@ -84,6 +84,12 @@ void UUVToolLivePreviewAPI::Initialize(UWorld* WorldIn, UInputRouter* RouterIn,
 	GetLivePreviewCameraStateFunc = MoveTemp(GetLivePreviewCameraStateFuncIn);
 }
 
+void UUVToolLivePreviewAPI::OnToolEnded(UInteractiveTool* DeadTool)
+{
+	OnDrawHUD.RemoveAll(DeadTool);
+	OnRender.RemoveAll(DeadTool);
+}
+
 int32 FUDIMBlock::BlockU() const
 {
 	int32 BlockU, BlockV;
diff --git a/Engine/Plugins/Experimental/UVEditor/Source/UVEditorTools/Private/Selection/UVToolSelectionAPI.cpp b/Engine/Plugins/Experimental/UVEditor/Source/UVEditorTools/Private/Selection/UVToolSelectionAPI.cpp
--- a/Engine/Plugins/Experimental/UVEditor/Source/UVEditorTools/Private/Selection/UVToolSelectionAPI.cpp
+++ b/Engine/Plugins/Experimental/UVEditor/Source/UVEditorTools/Private/Selection/UVToolSelectionAPI.cpp
@@ -410,6 +410,14 @@ void UUVToolSelectionAPI::DrawHUD(FCanvas* Canvas, IToolsContextRenderAPI* Rende
 		SelectionMechanic->DrawHUD(Canvas, RenderAPI);
 	}
 }
+void UUVToolSelectionAPI::LivePreviewRender(IToolsContextRenderAPI* RenderAPI)
+{
+	//TODO: here we'd route the render call to whatever mechanic is dealing with selection in live preview
+}
+void UUVToolSelectionAPI::LivePreviewDrawHUD(FCanvas* Canvas, IToolsContextRenderAPI* RenderAPI)
+{
+	//TODO: here we'd route the DrawHUD call to whatever mechanic is dealing with selection in live preview
+}
 
 void UUVToolSelectionAPI::BeginChange()
 {
diff --git a/Engine/Plugins/Experimental/UVEditor/Source/UVEditorTools/Public/ContextObjects/UVToolContextObjects.h b/Engine/Plugins/Experimental/UVEditor/Source/UVEditorTools/Public/ContextObjects/UVToolContextObjects.h
--- a/Engine/Plugins/Experimental/UVEditor/Source/UVEditorTools/Public/ContextObjects/UVToolContextObjects.h
+++ b/Engine/Plugins/Experimental/UVEditor/Source/UVEditorTools/Public/ContextObjects/UVToolContextObjects.h
@@ -119,6 +119,22 @@ public:
 		}
 	}
 
+	virtual void OnToolEnded(UInteractiveTool* DeadTool) override;
+
+	/**
+	 * Broadcast by the 3D live preview viewport on Render() so that mechanics/tools can
+	 * render there.
+	 */
+	DECLARE_MULTICAST_DELEGATE_OneParam(FOnRender, IToolsContextRenderAPI* RenderAPI);
+	FOnRender OnRender;
+
+	/**
+	 * Broadcast by the 3D live preview viewport on DrawHUD() so that mechanics/tools can
+	 * draw there.
+	 */
+	DECLARE_MULTICAST_DELEGATE_TwoParams(FOnDrawHUD, FCanvas* Canvas, IToolsContextRenderAPI* RenderAPI);
+	FOnDrawHUD OnDrawHUD;
+
 protected:
 	UPROPERTY()
 	TWeakObjectPtr<UWorld> World;
diff --git a/Engine/Plugins/Experimental/UVEditor/Source/UVEditorTools/Public/Selection/UVToolSelectionAPI.h b/Engine/Plugins/Experimental/UVEditor/Source/UVEditorTools/Public/Selection/UVToolSelectionAPI.h
--- a/Engine/Plugins/Experimental/UVEditor/Source/UVEditorTools/Public/Selection/UVToolSelectionAPI.h
+++ b/Engine/Plugins/Experimental/UVEditor/Source/UVEditorTools/Public/Selection/UVToolSelectionAPI.h
@@ -274,6 +274,10 @@ public:
 	// Called by the owner of the API
 	virtual void Render(IToolsContextRenderAPI* RenderAPI);
 	virtual void DrawHUD(FCanvas* Canvas, IToolsContextRenderAPI* RenderAPI);
+	// Could have had the api place these into the LivePreviewAPI delegates, but decided against
+	// it to align with Render() and DrawHUD() above. The owner will just call these directly.
+	virtual void LivePreviewRender(IToolsContextRenderAPI* RenderAPI);
+	virtual void LivePreviewDrawHUD(FCanvas* Canvas, IToolsContextRenderAPI* RenderAPI);
 
 	// UUVToolContextObject
 	virtual void Shutdown() override;
