UVEditor: Remove UVIslandConformalUnwrap action, now that its functionality is superseded by the Unwrap tool having selection support.

diff --git a/Engine/Plugins/Experimental/UVEditor/Source/UVEditor/Private/UVEditorCommands.cpp b/Engine/Plugins/Experimental/UVEditor/Source/UVEditor/Private/UVEditorCommands.cpp
--- a/Engine/Plugins/Experimental/UVEditor/Source/UVEditor/Private/UVEditorCommands.cpp
+++ b/Engine/Plugins/Experimental/UVEditor/Source/UVEditor/Private/UVEditorCommands.cpp
@@ -39,7 +39,6 @@
 	UI_COMMAND(SplitAction, "Split",
 	           "Given an edge selection, split those edges. Given a vertex selection, split any selected bowtie vertices.",
 	           EUserInterfaceActionType::Button, FInputChord());
-	UI_COMMAND(IslandConformalUnwrapAction, "IslUnw", "Apply a conformal unwrap to selected UV islands", EUserInterfaceActionType::Button, FInputChord());
 
 	// These allow us to link up to pressed keys
 	UI_COMMAND(AcceptOrCompleteActiveTool, "Accept", "Accept the active tool", EUserInterfaceActionType::Button, FInputChord(EKeys::Enter));
diff --git a/Engine/Plugins/Experimental/UVEditor/Source/UVEditor/Private/UVEditorMode.cpp b/Engine/Plugins/Experimental/UVEditor/Source/UVEditor/Private/UVEditorMode.cpp
--- a/Engine/Plugins/Experimental/UVEditor/Source/UVEditor/Private/UVEditorMode.cpp
+++ b/Engine/Plugins/Experimental/UVEditor/Source/UVEditor/Private/UVEditorMode.cpp
@@ -3,7 +3,6 @@
 #include "UVEditorMode.h"
 
 #include "Algo/AnyOf.h"
-#include "Actions/UVIslandConformalUnwrapAction.h"
 #include "Actions/UVSplitAction.h"
 #include "Actions/UVSeamSewAction.h"
 #include "Actions/UVToolAction.h"
@@ -471,7 +470,6 @@
 	};
 	PrepAction(CommandInfos.SewAction, NewObject<UUVSeamSewAction>());
 	PrepAction(CommandInfos.SplitAction, NewObject<UUVSplitAction>());
-	PrepAction(CommandInfos.IslandConformalUnwrapAction, NewObject<UUVIslandConformalUnwrapAction>());
 }
 
 bool UUVEditorMode::ShouldToolStartBeAllowed(const FString& ToolIdentifier) const
diff --git a/Engine/Plugins/Experimental/UVEditor/Source/UVEditor/Private/UVEditorModeToolkit.cpp b/Engine/Plugins/Experimental/UVEditor/Source/UVEditor/Private/UVEditorModeToolkit.cpp
--- a/Engine/Plugins/Experimental/UVEditor/Source/UVEditor/Private/UVEditorModeToolkit.cpp
+++ b/Engine/Plugins/Experimental/UVEditor/Source/UVEditor/Private/UVEditorModeToolkit.cpp
@@ -378,7 +378,6 @@
 
 	if (PaletteIndex == ToolsTabName)
 	{
-		ToolbarBuilder.AddToolBarButton(Commands.IslandConformalUnwrapAction);
 		ToolbarBuilder.AddToolBarButton(Commands.SewAction);
 		ToolbarBuilder.AddToolBarButton(Commands.SplitAction);
 
diff --git a/Engine/Plugins/Experimental/UVEditor/Source/UVEditor/Public/UVEditorCommands.h b/Engine/Plugins/Experimental/UVEditor/Source/UVEditor/Public/UVEditorCommands.h
--- a/Engine/Plugins/Experimental/UVEditor/Source/UVEditor/Public/UVEditorCommands.h
+++ b/Engine/Plugins/Experimental/UVEditor/Source/UVEditor/Public/UVEditorCommands.h
@@ -29,7 +29,6 @@
 
 	TSharedPtr<FUICommandInfo> SewAction;
 	TSharedPtr<FUICommandInfo> SplitAction;
-	TSharedPtr<FUICommandInfo> IslandConformalUnwrapAction;
 
 	TSharedPtr<FUICommandInfo> AcceptOrCompleteActiveTool;
 	TSharedPtr<FUICommandInfo> CancelOrCompleteActiveTool;
diff --git a/Engine/Plugins/Experimental/UVEditor/Source/UVEditorTools/Private/Actions/UVIslandConformalUnwrapAction.cpp b/Engine/Plugins/Experimental/UVEditor/Source/UVEditorTools/Private/Actions/UVIslandConformalUnwrapAction.cpp
deleted file mode 100644
diff --git a/Engine/Plugins/Experimental/UVEditor/Source/UVEditorTools/Public/Actions/UVIslandConformalUnwrapAction.h b/Engine/Plugins/Experimental/UVEditor/Source/UVEditorTools/Public/Actions/UVIslandConformalUnwrapAction.h
deleted file mode 100644
