UVEditor: Make UV editor accessible for supported assets from right-click context menu

diff --git a/Engine/Plugins/Experimental/UVEditor/Source/UVEditor/Private/UVEditorModule.cpp b/Engine/Plugins/Experimental/UVEditor/Source/UVEditor/Private/UVEditorModule.cpp
--- a/Engine/Plugins/Experimental/UVEditor/Source/UVEditor/Private/UVEditorModule.cpp
+++ b/Engine/Plugins/Experimental/UVEditor/Source/UVEditor/Private/UVEditorModule.cpp
@@ -100,55 +100,93 @@ void FUVEditorModule::RegisterMenus()
 		AddToContextMenuSection(Section);
 	}
 
-	// Extend the level editor context menu
+	// Extend the level editor menus
+
+	// Helper struct to collect all necessary inputs for a menu entry.
+	struct FMenuEntryParameters
 	{
-		UToolMenu* Menu = UToolMenus::Get()->ExtendMenu("LevelEditor.ActorContextMenu.AssetToolsSubMenu");
+		TArray<TObjectPtr<UObject>> TargetObjects;
+		UUVEditorSubsystem* UVSubsystem;
+		bool bValidTargets;
+	};
 
-		FToolMenuSection& Section = Menu->AddSection("UVEditorCommands", TAttribute<FText>(), FToolMenuInsert("AssetTools", EToolMenuInsertType::After));
+	// Sets up all parameters for a menu entry.
+	auto SetupMenuEntryParameters = []
+	{
+		FMenuEntryParameters Result;
 
-		Section.AddDynamicEntry("OpenUVEditor", FNewToolMenuSectionDelegate::CreateLambda(
-			[this](FToolMenuSection& Section) 
-			{
-				TArray<TObjectPtr<UObject>> TargetObjects;
+		for (FSelectionIterator It(GEditor->GetSelectedActorIterator()); It; ++It)
+		{
+			// We are interested in the (unique) assets backing the actor, or else the actor
+			// itself if it is not asset backed (such as UDynamicMesh).
+			const AActor* Actor = static_cast<AActor*>(*It);
+			TArray<UObject*> ActorAssets;
+			Actor->GetReferencedContentObjects(ActorAssets);
 
-				// TODO: There's some newer way to iterate across selected actors that we can switch over to someday
-				for (FSelectionIterator It(GEditor->GetSelectedActorIterator()); It; ++It)
+			if (ActorAssets.Num() > 0)
+			{
+				for (UObject* Asset : ActorAssets)
 				{
-					// We are interested in the (unique) assets backing the actor, or else the actor
-					// itself if it is not asset backed (such as UDynamicMesh).
-					AActor* Actor = static_cast<AActor*>(*It);
-					TArray<UObject*> ActorAssets;
-					Actor->GetReferencedContentObjects(ActorAssets);
-
-					if (ActorAssets.Num() > 0)
-					{
-						for (UObject* Asset : ActorAssets)
-						{
-							TargetObjects.AddUnique(Asset);
-						}
-					}
-					else
-					{
-						// Need to transform actors to components here because this is what our tool targets expect
-						TInlineComponentArray<UActorComponent*> ActorComponents;
-						Actor->GetComponents(ActorComponents);
-						TargetObjects.Append(ActorComponents);
-					}
+					Result.TargetObjects.AddUnique(Asset);
 				}
+			}
+			else
+			{
+				// Need to transform actors to components here because this is what our tool targets expect
+				TInlineComponentArray<UActorComponent*> ActorComponents;
+				Actor->GetComponents(ActorComponents);
+				Result.TargetObjects.Append(ActorComponents);
+			}
+		}
 
-				UUVEditorSubsystem* UVSubsystem = GEditor->GetEditorSubsystem<UUVEditorSubsystem>();
-				check(UVSubsystem);
+		Result.UVSubsystem = GEditor->GetEditorSubsystem<UUVEditorSubsystem>();
+		check(Result.UVSubsystem);
 
-				bool bValidTargets = UVSubsystem->AreObjectsValidTargets(TargetObjects);
+		Result.bValidTargets = Result.UVSubsystem->AreObjectsValidTargets(Result.TargetObjects);
 
-				TSharedPtr<class FUICommandList> CommandListToBind = MakeShared<FUICommandList>();
-				CommandListToBind->MapAction(
-					FUVEditorCommands::Get().OpenUVEditor,
-					FExecuteAction::CreateUObject(UVSubsystem, &UUVEditorSubsystem::StartUVEditor, TargetObjects),
-					FCanExecuteAction::CreateLambda([bValidTargets]() { return bValidTargets; }));
+		return Result;
+	};
 
-				Section.AddMenuEntryWithCommandList(FUVEditorCommands::Get().OpenUVEditor, CommandListToBind);
-			}));
+	// Adds a menu entry for the given parameters.
+	auto AddMenuEntry = [](FToolMenuSection& Section, const FMenuEntryParameters& MenuEntryParameters)
+	{
+		const bool bValidTargets = MenuEntryParameters.bValidTargets;
+
+		const TSharedPtr<class FUICommandList> CommandListToBind = MakeShared<FUICommandList>();
+		CommandListToBind->MapAction(
+			FUVEditorCommands::Get().OpenUVEditor,
+			FExecuteAction::CreateUObject(MenuEntryParameters.UVSubsystem, &UUVEditorSubsystem::StartUVEditor, MenuEntryParameters.TargetObjects),
+			FCanExecuteAction::CreateLambda([bValidTargets]() { return bValidTargets; }));
+
+		Section.AddMenuEntryWithCommandList(FUVEditorCommands::Get().OpenUVEditor, CommandListToBind);
+	};
+
+	// Add UV Editor to "Actor" > "Asset Tools" sub menu.
+	{
+		UToolMenu* AssetToolsSubMenu = UToolMenus::Get()->ExtendMenu("LevelEditor.ActorContextMenu.AssetToolsSubMenu");
+		FToolMenuSection& AssetToolsSection = AssetToolsSubMenu->AddSection("UVEditorCommands", TAttribute<FText>(), FToolMenuInsert("AssetTools",
+			                                                                    EToolMenuInsertType::After));
+		AssetToolsSection.AddDynamicEntry("OpenUVEditor", FNewToolMenuSectionDelegate::CreateLambda(
+			                                  [&SetupMenuEntryParameters, &AddMenuEntry](FToolMenuSection& Section)
+			                                  {
+				                                  const FMenuEntryParameters MenuEntryParameters = SetupMenuEntryParameters();
+				                                  AddMenuEntry(Section, MenuEntryParameters);
+			                                  }));
+	}
+
+	// Add UV Editor to actor context menu.
+	{
+		UToolMenu* ActorContextMenu = UToolMenus::Get()->ExtendMenu("LevelEditor.ActorContextMenu");
+		FToolMenuSection& ActorTypeToolsSection = ActorContextMenu->FindOrAddSection("ActorTypeTools");
+		ActorTypeToolsSection.AddDynamicEntry("OpenUVEditor", FNewToolMenuSectionDelegate::CreateLambda(
+			                                      [&SetupMenuEntryParameters, &AddMenuEntry](FToolMenuSection& Section)
+			                                      {
+				                                      const FMenuEntryParameters MenuEntryParameters = SetupMenuEntryParameters();
+				                                      if (MenuEntryParameters.bValidTargets)
+				                                      {
+					                                      AddMenuEntry(Section, MenuEntryParameters);
+				                                      }
+			                                      }));
 	}
 }
 
