UVEditor: Add support for the split action to handle triangle based selections by splitting along selection region boundaries.

diff --git a/Engine/Plugins/Experimental/UVEditor/Source/UVEditor/Private/UVEditorCommands.cpp b/Engine/Plugins/Experimental/UVEditor/Source/UVEditor/Private/UVEditorCommands.cpp
--- a/Engine/Plugins/Experimental/UVEditor/Source/UVEditor/Private/UVEditorCommands.cpp
+++ b/Engine/Plugins/Experimental/UVEditor/Source/UVEditor/Private/UVEditorCommands.cpp
@@ -37,7 +37,7 @@
 	// once selection is pulled out to mode-level.
 	UI_COMMAND(SewAction, "Sew", "Sew edges highlighted in red to edges highlighted in green", EUserInterfaceActionType::Button, FInputChord());
 	UI_COMMAND(SplitAction, "Split",
-	           "Given an edge selection, split those edges. Given a vertex selection, split any selected bowtie vertices.",
+	           "Given an edge selection, split those edges. Given a vertex selection, split any selected bowtie vertices. Given a triangle selection, split along selection boundaries.",
 	           EUserInterfaceActionType::Button, FInputChord());
 
 	// These allow us to link up to pressed keys
diff --git a/Engine/Plugins/Experimental/UVEditor/Source/UVEditorTools/Private/Actions/UVSplitAction.cpp b/Engine/Plugins/Experimental/UVEditor/Source/UVEditorTools/Private/Actions/UVSplitAction.cpp
--- a/Engine/Plugins/Experimental/UVEditor/Source/UVEditorTools/Private/Actions/UVSplitAction.cpp
+++ b/Engine/Plugins/Experimental/UVEditor/Source/UVEditorTools/Private/Actions/UVSplitAction.cpp
@@ -153,13 +153,102 @@ namespace UVSplitActionLocals
 
 		return true;
 	}
+
+	bool ApplyConnectedRegionSplits(const FUVToolSelection& Selection, FUVToolSelection& NewSelectionOut,
+		UUVToolEmitChangeAPI& EmitChangeAPI, const FText& TransactionName)
+	{
+		UUVEditorToolMeshInput* Target = Selection.Target.Get();
+
+		if (!ensure(Target && Target->IsValid()))
+		{
+			return false;
+		}
+
+		TSet<int32> UnwrapRegionBoundaryEdges;
+		for (int32 Tid : Selection.SelectedIDs)
+		{
+			if (!ensure(Target->UnwrapCanonical->IsTriangle(Tid)))
+			{
+				continue;
+			}
+
+			// Gather the boundary edges for each region
+			FIndex3i TriEids = Target->UnwrapCanonical->GetTriEdges(Tid);
+			for (int i = 0; i < 3; ++i)
+			{
+				FIndex2i EdgeTids = Target->UnwrapCanonical->GetEdgeT(TriEids[i]);
+				for (int j = 0; j < 2; ++j)
+				{
+					if (EdgeTids[j] != Tid && !Selection.SelectedIDs.Contains(EdgeTids[j]))
+					{
+						UnwrapRegionBoundaryEdges.Add(TriEids[i]);
+						break;
+					}
+				}
+			}//end for tri edges
+		}//end for selection tids
+
+		TSet<int32> AppliedEidSet;
+		for (int32 Eid : UnwrapRegionBoundaryEdges)
+		{
+			// Note that we don't check whether edges are already boundary edges because we allow such edges
+			// to be selected for splitting of any attached bowties.
+
+			FIndex2i EdgeUnwrapVids = Target->UnwrapCanonical->GetEdgeV(Eid);
+
+			int32 AppliedEid = Target->AppliedCanonical->FindEdge(
+				Target->UnwrapVidToAppliedVid(EdgeUnwrapVids.A),
+				Target->UnwrapVidToAppliedVid(EdgeUnwrapVids.B));
+
+			if (ensure(AppliedEid != IndexConstants::InvalidID))
+			{
+				AppliedEidSet.Add(AppliedEid);
+			}
+		}
+
+		// Perform the cut in the overlay, but don't propagate to unwrap yet because we'll need to prep for undo
+		FUVEditResult UVEditResult;
+		FDynamicMeshUVEditor UVEditor(Target->AppliedCanonical.Get(),
+			Target->UVLayerIndex, false);
+		UVEditor.CreateSeamsAtEdges(AppliedEidSet, &UVEditResult);
+
+		// Figure out the triangles that need to be saved in the unwrap for undo
+		TSet<int32> TidSet;
+		for (int32 UnwrapVid : UVEditResult.NewUVElements)
+		{
+			TArray<int32> VertTids;
+			Target->AppliedCanonical->GetVtxTriangles(Target->UnwrapVidToAppliedVid(UnwrapVid), VertTids);
+			TidSet.Append(VertTids);
+		}
+
+		FDynamicMeshChangeTracker ChangeTracker(Target->UnwrapCanonical.Get());
+		ChangeTracker.BeginChange();
+		ChangeTracker.SaveTriangles(TidSet, true);
+
+		// Perform the update
+		TArray<int32> AppliedTids = TidSet.Array();
+		Target->UpdateAllFromAppliedCanonical(&UVEditResult.NewUVElements, &AppliedTids, &AppliedTids);
+
+		// Since we are simply splitting triangles away, the selection doesn't actually change after the split.
+		NewSelectionOut.Target = Selection.Target;
+		NewSelectionOut.Type = Selection.Type;
+		NewSelectionOut.SelectedIDs = Selection.SelectedIDs;
+
+		// Emit update transaction
+		EmitChangeAPI.EmitToolIndependentUnwrapCanonicalChange(
+			Target, ChangeTracker.EndChange(), TransactionName);
+
+		return true;
+
+	}
 }
 
 bool UUVSplitAction::CanExecuteAction() const
 {
 	return SelectionAPI->HaveSelections()
 		&& (SelectionAPI->GetSelectionsType() == FUVToolSelection::EType::Vertex
-			|| SelectionAPI->GetSelectionsType() == FUVToolSelection::EType::Edge);
+			|| SelectionAPI->GetSelectionsType() == FUVToolSelection::EType::Edge
+			|| SelectionAPI->GetSelectionsType() == FUVToolSelection::EType::Triangle);
 }
 
 bool UUVSplitAction::ExecuteAction()
@@ -176,9 +265,19 @@ bool UUVSplitAction::ExecuteAction()
 	FUVToolSelection::EType SelectionType = SelectionAPI->GetSelectionsType();
 	TArray<FUVToolSelection> NewSelections;
 
-	const FText TransactionName = (SelectionType == FUVToolSelection::EType::Edge) ?
-		LOCTEXT("SplitEdgesTransactionName", "Split Edges")
-		: LOCTEXT("SplitBowtieVerticesTransactionName", "Split Bowties");
+	FText TransactionName;
+	if (SelectionType == FUVToolSelection::EType::Edge)
+	{
+		TransactionName = LOCTEXT("SplitEdgesTransactionName", "Split Edges");
+	}
+	else if(SelectionType == FUVToolSelection::EType::Vertex)
+	{
+		TransactionName = LOCTEXT("SplitBowtieVerticesTransactionName", "Split Bowties");
+	}
+	else if (SelectionType == FUVToolSelection::EType::Triangle)
+	{
+		TransactionName = LOCTEXT("SplitRegionsTransactionName", "Split Regions");
+	}
 	EmitChangeAPI->BeginUndoTransaction(TransactionName);
 
 	SelectionAPI->ClearSelections(false, true); // don't broadcast, do emit
@@ -193,6 +292,10 @@ bool UUVSplitAction::ExecuteAction()
 		{
 			bAllSucceeded &= ApplySplitBowtieVertices(Selection, NewSelection, *EmitChangeAPI, TransactionName);
 		}
+		else if (Selection.Type == FUVToolSelection::EType::Triangle)
+		{
+			bAllSucceeded &= ApplyConnectedRegionSplits(Selection, NewSelection, *EmitChangeAPI, TransactionName);
+		}
 
 		if (!NewSelection.IsEmpty())
 		{
