UV Editor: Enable UDIM support CVAR by default

#rnx

diff --git a/Engine/Plugins/Experimental/UVEditor/Source/UVEditorTools/Private/UVEditorUXSettings.cpp b/Engine/Plugins/Experimental/UVEditor/Source/UVEditorTools/Private/UVEditorUXSettings.cpp
--- a/Engine/Plugins/Experimental/UVEditor/Source/UVEditorTools/Private/UVEditorUXSettings.cpp
+++ b/Engine/Plugins/Experimental/UVEditor/Source/UVEditorTools/Private/UVEditorUXSettings.cpp
@@ -88,7 +88,7 @@ const int32 FUVEditorUXSettings::RulerSubdivisionLevel(1);
 
 TAutoConsoleVariable<int32> FUVEditorUXSettings::CVarEnablePrototypeUDIMSupport(
 	TEXT("modeling.UVEditor.UDIMSupport"),
-	0,
+	1,
 	TEXT("Enable experimental UDIM support in the UVEditor"));
 
 
