UVEditor: Enable support for selection mechanics within in the 3D viewport.

diff --git a/Engine/Plugins/Experimental/UVEditor/Source/UVEditor/Private/UVEditorMode.cpp b/Engine/Plugins/Experimental/UVEditor/Source/UVEditor/Private/UVEditorMode.cpp
--- a/Engine/Plugins/Experimental/UVEditor/Source/UVEditor/Private/UVEditorMode.cpp
+++ b/Engine/Plugins/Experimental/UVEditor/Source/UVEditor/Private/UVEditorMode.cpp
@@ -500,11 +500,26 @@
 	
 		// If a tool doesn't support selection, we can't be certain that it won't put the meshes
 		// in a state where the selection refers to invalid elements.
-		if (!Cast<IUVToolSupportsSelection>(Tool))
+		IUVToolSupportsSelection* SelectionTool = Cast<IUVToolSupportsSelection>(Tool);
+		if (!SelectionTool)
 		{
-			SelectionAPI->ClearSelections(true, true); // broadcast, emit
+			SelectionAPI->BeginChange();
+			SelectionAPI->ClearSelections(false, false); 
+			SelectionAPI->ClearUnsetElementAppliedMeshSelections(false, false);
+			SelectionAPI->EndChangeAndEmitIfModified(true); // broadcast, emit
 			SelectionAPI->ClearHighlight();
 		}
+		else
+		{
+			if (!SelectionTool->SupportsUnsetElementAppliedMeshSelections())
+			{
+				SelectionAPI->BeginChange();				
+				SelectionAPI->ClearUnsetElementAppliedMeshSelections(false, false);
+				SelectionAPI->EndChangeAndEmitIfModified(true); // broadcast, emit
+				SelectionAPI->ClearHighlight(false, true); // Clear and rebuild applied highlight to account for clearing unset selections
+				SelectionAPI->RebuildAppliedPreviewHighlight();
+			}
+		}
 	
 		GetInteractiveToolsContext()->GetTransactionAPI()->AppendChange(
 			this, MakeUnique<FUVEditorBeginToolChange>(), TransactionName);
diff --git a/Engine/Plugins/Experimental/UVEditor/Source/UVEditor/Public/UVEditor3DViewportClient.h b/Engine/Plugins/Experimental/UVEditor/Source/UVEditor/Public/UVEditor3DViewportClient.h
--- a/Engine/Plugins/Experimental/UVEditor/Source/UVEditor/Public/UVEditor3DViewportClient.h
+++ b/Engine/Plugins/Experimental/UVEditor/Source/UVEditor/Public/UVEditor3DViewportClient.h
@@ -28,9 +28,18 @@ public:
 
 	// FEditorViewportClient
 	virtual bool ShouldOrbitCamera() const override {
+
+		// Including some additional checks to prevent the orbit mode from being on all the time,
+		// which ultimately causes weirdness in how the camera transform matrices are handled by
+		// the viewport internally.
+		const bool bLeftMouseButtonDown = Viewport->KeyState(EKeys::LeftMouseButton) && !bLockFlightCamera;
+		const bool bMiddleMouseButtonDown = Viewport->KeyState(EKeys::MiddleMouseButton);
+		const bool bRightMouseButtonDown = Viewport->KeyState(EKeys::RightMouseButton);
+		const bool bIsOnlyAltPressed = IsAltPressed() && !IsCtrlPressed() && !IsShiftPressed();
+
 		switch (CameraMode) {
 		case EUVEditor3DViewportClientCameraMode::Orbit:
-			return true;
+			return bIsOnlyAltPressed || bLeftMouseButtonDown || bMiddleMouseButtonDown || bRightMouseButtonDown;
 		case EUVEditor3DViewportClientCameraMode::Fly:
 			return FEditorViewportClient::ShouldOrbitCamera();
 		default:
diff --git a/Engine/Plugins/Experimental/UVEditor/Source/UVEditorTools/Private/Selection/UVEditorMeshSelectionMechanic.cpp b/Engine/Plugins/Experimental/UVEditor/Source/UVEditorTools/Private/Selection/UVEditorMeshSelectionMechanic.cpp
--- a/Engine/Plugins/Experimental/UVEditor/Source/UVEditorTools/Private/Selection/UVEditorMeshSelectionMechanic.cpp
+++ b/Engine/Plugins/Experimental/UVEditor/Source/UVEditorTools/Private/Selection/UVEditorMeshSelectionMechanic.cpp
@@ -5,6 +5,7 @@
 #include "BaseBehaviors/BehaviorTargetInterfaces.h"
 #include "BaseBehaviors/SingleClickOrDragBehavior.h"
 #include "BaseBehaviors/MouseHoverBehavior.h"
+#include "BaseBehaviors/SingleClickBehavior.h"
 #include "ContextObjectStore.h"
 #include "ContextObjects/UVToolContextObjects.h"
 #include "ContextObjects/UVToolViewportButtonsAPI.h"
@@ -83,64 +84,70 @@ namespace UVEditorMeshSelectionMechanicLocals
 		return {Point.X, Point.Y};
 	}
 
-	void AppendVertexIDs(const FDynamicMesh3& MeshXY0, int TriangleID, TArray<int>& VertexIDs)
+	void AppendVertexIDs(const FDynamicMesh3& Mesh, int TriangleID, TArray<int>& VertexIDs)
 	{
-		const FIndex3i& Triangle = MeshXY0.GetTriangleRef(TriangleID);
+		const FIndex3i& Triangle = Mesh.GetTriangleRef(TriangleID);
 		VertexIDs.Add(Triangle.A);
 		VertexIDs.Add(Triangle.B);
 		VertexIDs.Add(Triangle.C);
 	}
 
-	void AppendVertexIDsIfIntersected(const FDynamicMesh3& MeshXY0, const FAxisAlignedBox2d& RectangleXY, int TriangleID, TArray<int>& VertexIDs)
+	void AppendVertexIDsIfIntersectedRectangle(const FDynamicMesh3& MeshXY0, const FAxisAlignedBox2d& RectangleXY, int TriangleID, TArray<int>& VertexIDs)
 	{
 		const FIndex3i& Triangle = MeshXY0.GetTriangleRef(TriangleID);
-		if (RectangleXY.Contains(XY(MeshXY0.GetVertex(Triangle.A))))
-		{
-			VertexIDs.Add(Triangle.A);
-		}
-		
-		if (RectangleXY.Contains(XY(MeshXY0.GetVertex(Triangle.B))))
+		for (int32 Index = 0; Index < 3; ++Index)
 		{
-			VertexIDs.Add(Triangle.B);
+			if (RectangleXY.Contains(XY(MeshXY0.GetVertex(Triangle[Index]))))
+			{
+				VertexIDs.Add(Triangle[Index]);
+			}
 		}
-		
-		if (RectangleXY.Contains(XY(MeshXY0.GetVertex(Triangle.C))))
+	}
+
+	void AppendVertexIDsIfIntersectedFrustum(const FDynamicMesh3& Mesh, const FConvexVolume& Frustum, int TriangleID, TArray<int>& VertexIDs)
+	{
+		const FIndex3i& Triangle = Mesh.GetTriangleRef(TriangleID);
+		for (int32 Index = 0; Index < 3; ++Index)
 		{
-			VertexIDs.Add(Triangle.C);
+			if (Frustum.IntersectPoint(Mesh.GetVertex(Triangle[Index])))
+			{
+				VertexIDs.Add(Triangle[Index]);
+			}
 		}
 	}
 	
-	void AppendEdgeIDs(const FDynamicMesh3& MeshXY0, int TriangleID, TArray<int>& EdgeIDs)
+	void AppendEdgeIDs(const FDynamicMesh3& Mesh, int TriangleID, TArray<int>& EdgeIDs)
 	{
-		const FIndex3i& Edges = MeshXY0.GetTriEdgesRef(TriangleID);
+		const FIndex3i& Edges = Mesh.GetTriEdgesRef(TriangleID);
 		EdgeIDs.Add(Edges.A);
 		EdgeIDs.Add(Edges.B);
 		EdgeIDs.Add(Edges.C);
 	}
 	
-	void AppendEdgeIDsIfIntersected(const FDynamicMesh3& MeshXY0, const FAxisAlignedBox2d& RectangleXY, int TriangleID, TArray<int>& EdgeIDs)
+	void AppendEdgeIDsIfIntersectedRectangle(const FDynamicMesh3& MeshXY0, const FAxisAlignedBox2d& RectangleXY, int TriangleID, TArray<int>& EdgeIDs)
 	{
 		const FIndex3i& Edges = MeshXY0.GetTriEdgesRef(TriangleID);
-
-		const FIndex2i& EdgeA = MeshXY0.GetEdgeRef(Edges.A).Vert;
-		const FSegment2d SegmentA(XY(MeshXY0.GetVertex(EdgeA.A)), XY(MeshXY0.GetVertex(EdgeA.B)));
-		if (TestIntersection(SegmentA, RectangleXY))
-		{
-			EdgeIDs.Add(Edges.A);
-		}
-		
-		const FIndex2i& EdgeB = MeshXY0.GetEdgeRef(Edges.B).Vert;
-		const FSegment2d SegmentB(XY(MeshXY0.GetVertex(EdgeB.A)), XY(MeshXY0.GetVertex(EdgeB.B)));
-		if (TestIntersection(SegmentB, RectangleXY))
+		for (int32 Index = 0; Index < 3; ++Index)
 		{
-			EdgeIDs.Add(Edges.B);
+			const FIndex2i& EdgeVerts = MeshXY0.GetEdgeRef(Edges[Index]).Vert;
+			const FSegment2d Segment(XY(MeshXY0.GetVertex(EdgeVerts.A)), XY(MeshXY0.GetVertex(EdgeVerts.B)));
+			if (TestIntersection(Segment, RectangleXY))
+			{
+				EdgeIDs.Add(Edges[Index]);
+			}
 		}
-		
-		const FIndex2i& EdgeC = MeshXY0.GetEdgeRef(Edges.C).Vert;
-		const FSegment2d SegmentC(XY(MeshXY0.GetVertex(EdgeC.A)), XY(MeshXY0.GetVertex(EdgeC.B)));
-		if (TestIntersection(SegmentC, RectangleXY))
+	}
+
+	void AppendEdgeIDsIfIntersectedFrustum(const FDynamicMesh3& Mesh, const FConvexVolume& Frustum, int TriangleID, TArray<int>& EdgeIDs)
+	{
+		const FIndex3i& Edges = Mesh.GetTriEdgesRef(TriangleID);
+		for (int32 Index = 0; Index < 3; ++Index)
 		{
-			EdgeIDs.Add(Edges.C);
+			const FIndex2i& EdgeVerts = Mesh.GetEdgeRef(Edges[Index]).Vert;
+			if (Frustum.IntersectLineSegment(Mesh.GetVertex(EdgeVerts.A), Mesh.GetVertex(EdgeVerts.B)))
+			{
+				EdgeIDs.Add(Edges[Index]);
+			}
 		}
 	}
 	
@@ -149,7 +156,7 @@ namespace UVEditorMeshSelectionMechanicLocals
 		TriangleIDs.Add(TriangleID);
 	}
 
-	void AppendTriangleIDIfIntersected(const FDynamicMesh3& MeshXY0, const FAxisAlignedBox2d& RectangleXY, int TriangleID, TArray<int>& TriangleIDs)
+	void AppendTriangleIDIfIntersectedRectangle(const FDynamicMesh3& MeshXY0, const FAxisAlignedBox2d& RectangleXY, int TriangleID, TArray<int>& TriangleIDs)
 	{
 		const FIndex3i& Triangle = MeshXY0.GetTriangleRef(TriangleID);
 		const FTriangle2d TriangleXY(XY(MeshXY0.GetVertex(Triangle.A)),
@@ -162,30 +169,112 @@ namespace UVEditorMeshSelectionMechanicLocals
 			TriangleIDs.Add(TriangleID);
 		}
 	}
-	
-	// Returns indices, collected by the given functions, from triangles which are intersected by the given rectangle.
-	// TreeXY0 must contain a mesh with vertices in the XY plane (have zero Z coordinate)
-	template<typename IDsFromTriangleF, typename IDsFromTriangleIfIntersectedF>
-	TArray<int32> FindAllIntersectionsAxisAlignedBox2(const FDynamicMeshAABBTree3& TreeXY0,
-													  const FAxisAlignedBox2d& RectangleXY,
-													  IDsFromTriangleF AppendIDs,
-													  IDsFromTriangleIfIntersectedF AppendIDsIfIntersected)
+
+	void AppendTriangleIDIfIntersectedFrustum(const FDynamicMesh3& Mesh, const FConvexVolume& Frustum, int TriangleID, TArray<int>& TriangleIDs)
+	{
+		FVector3d V1, V2, V3;
+		bool bFullyContained;
+		Mesh.GetTriVertices(TriangleID, V1, V2, V3);
+
+		if (Frustum.IntersectTriangle((FVector)V1, (FVector)V2, (FVector)V3, bFullyContained))
+		{
+			TriangleIDs.Add(TriangleID);
+		}
+	}
+
+	/*
+	 * The tree processing in FindAllIntersections works for both 2d rectangles and 3d frustums as long as we
+	 * provide the proper final intersection test. This template is a helper to allow us to use the proper
+	 * intersection test for either type.
+	*/ 
+	template<typename RegionImpl>
+	struct TTreeBoxIntersectHelper;
+
+	template<>
+	struct TTreeBoxIntersectHelper<FAxisAlignedBox2d>
+	{
+		typedef FVector2d PointType;
+		typedef FAxisAlignedBox2d BoxType;
+
+		static BoxType BuildBoxFromTree(const FDynamicMeshAABBTree3& Tree)
+		{
+			BoxType Box;
+			Box.Contain(XY(Tree.GetBoundingBox().Min));
+			Box.Contain(XY(Tree.GetBoundingBox().Max));
+			return Box;
+		}
+
+		static BoxType BuildBoxFromBox3d(const FAxisAlignedBox3d& Box3d)
+		{
+			return BoxType(XY(Box3d.Min), XY(Box3d.Max));
+		}
+
+		static bool RegionBoxTest(const FAxisAlignedBox2d& Region, const FAxisAlignedBox2d& Box, bool& bFullyContained)
+		{
+			bFullyContained = false;
+			if (Region.Intersects(Box))
+			{
+				if (Region.Contains(Box))
+				{
+					bFullyContained = true;
+				}
+				return true;
+			}
+			return false;
+		}
+	};
+
+	template<>
+	struct TTreeBoxIntersectHelper<FConvexVolume>
+	{
+		typedef FVector3d PointType;
+		typedef FAxisAlignedBox3d BoxType;
+
+		static BoxType BuildBoxFromTree(const FDynamicMeshAABBTree3& Tree)
+		{
+			BoxType Box;
+			Box.Contain(Tree.GetBoundingBox().Min);
+			Box.Contain(Tree.GetBoundingBox().Max);
+			return Box;
+		}
+
+		static BoxType BuildBoxFromBox3d(const FAxisAlignedBox3d& Box3d)
+		{
+			return Box3d;
+		}
+
+		static bool RegionBoxTest(const FConvexVolume& Region, const FAxisAlignedBox3d& Box, bool& bFullyContained)
+		{
+			return Region.IntersectBox(Box.Center(), Box.Extents(), bFullyContained);
+		}
+	};
+
+	// Returns indices, collected by the given functions, from triangles which are intersected by the given region,
+	// which can be either a rectangle (in which case the underlying mesh should have its vertices in the XY plane
+	// with a zero Z coordinate) or a frustum. 
+	template<typename IDsFromTriangleF, typename IDsFromTriangleIfIntersectedF, typename RegionType>
+	TArray<int32> FindAllIntersections(const FDynamicMeshAABBTree3& Tree,
+									   const RegionType& Region,
+									   IDsFromTriangleF AppendIDs,
+									   IDsFromTriangleIfIntersectedF AppendIDsIfIntersected)
 	{
 		TRACE_CPUPROFILER_EVENT_SCOPE(FindAllIntersectionsAxisAlignedBox2);
-		
-		check(TreeXY0.GetMesh());
+		typedef TTreeBoxIntersectHelper<RegionType> RegionPolicy;
+
+		check(Tree.GetMesh());
 		
 		TArray<int32> Result;
-		FAxisAlignedBox2d TreeRectangleXY;
-		TreeRectangleXY.Contain(XY(TreeXY0.GetBoundingBox().Min));
-		TreeRectangleXY.Contain(XY(TreeXY0.GetBoundingBox().Max));
-		if (RectangleXY.Contains(TreeRectangleXY))
+		typename RegionPolicy::BoxType TreeBox = RegionPolicy::BuildBoxFromTree(Tree);
+
+		bool bFullyContained;
+		bool bIntersects = RegionPolicy::RegionBoxTest(Region, TreeBox, bFullyContained);
+		if (bFullyContained)
 		{
 			// Early out selecting everything
-			Result.Reserve(TreeXY0.GetMesh()->TriangleCount());
-			for (int TriangleID : TreeXY0.GetMesh()->TriangleIndicesItr())
+			Result.Reserve(Tree.GetMesh()->TriangleCount());
+			for (int TriangleID : Tree.GetMesh()->TriangleIndicesItr())
 			{
-				AppendIDs(*TreeXY0.GetMesh(), TriangleID, Result);
+				AppendIDs(*Tree.GetMesh(), TriangleID, Result);
 			}
 			return Result;
 		}
@@ -197,7 +286,7 @@ namespace UVEditorMeshSelectionMechanicLocals
 		FDynamicMeshAABBTree3::FTreeTraversal Traversal;
 		
 		Traversal.NextBoxF =
-			[&RectangleXY, &SelectAllDepth, &CurrentDepth](const FAxisAlignedBox3d& Box, int Depth)
+			[&Region, &SelectAllDepth, &CurrentDepth](const FAxisAlignedBox3d& Box3d, int Depth)
 		{
 			CurrentDepth = Depth;
 			if (Depth > SelectAllDepth)
@@ -208,39 +297,39 @@ namespace UVEditorMeshSelectionMechanicLocals
 			}
 			
 			SelectAllDepth = TNumericLimits<int>::Max();
-			
-			const FAxisAlignedBox2d BoxXY(XY(Box.Min), XY(Box.Max));
-			if (RectangleXY.Intersects(BoxXY))
+			typename RegionPolicy::BoxType RegionBox = RegionPolicy::BuildBoxFromBox3d(Box3d);
+
+			bool bFullyContained;
+			bool bIntersects = RegionPolicy::RegionBoxTest(Region, RegionBox, bFullyContained);
+
+			if(bFullyContained)
 			{
-				if (RectangleXY.Contains(BoxXY))
-				{
-					SelectAllDepth = Depth;
-				}
-				
-				return true;		
+				SelectAllDepth = Depth;
 			}
-			return false;
+				
+			return bIntersects;
 		};
 		
 		Traversal.NextTriangleF =
-			[&RectangleXY, &SelectAllDepth, &CurrentDepth, &TreeXY0, &Result, &AppendIDs, &AppendIDsIfIntersected]
+			[&Region, &SelectAllDepth, &CurrentDepth, &Tree, &Result, &AppendIDs, &AppendIDsIfIntersected]
 			(int TriangleID)
 		{
 			if (CurrentDepth >= SelectAllDepth)
 			{
 				// This TriangleID is entirely contained in the selection rectangle so we can skip intersection testing
-				return AppendIDs(*TreeXY0.GetMesh(), TriangleID, Result);
+				return AppendIDs(*Tree.GetMesh(), TriangleID, Result);
 			}
-			return AppendIDsIfIntersected(*TreeXY0.GetMesh(), RectangleXY, TriangleID, Result);
+			return AppendIDsIfIntersected(*Tree.GetMesh(), Region, TriangleID, Result);
 		};
 		
-		TreeXY0.DoTraversal(Traversal);
+		Tree.DoTraversal(Traversal);
 
 		return Result;
 	}
 
-	bool ConvertToHitElementList(ESelectionMode SelectionMode, 
-		const FDynamicMesh3& Mesh, int32 HitTid, const FViewCameraState& CameraState, 
+
+	bool ConvertToHitElementList(ESelectionMode SelectionMode,
+		const FDynamicMesh3& Mesh, const FDynamicMesh3& UnwrapMesh, FDynamicMeshUVOverlay& UVOverlay, int32 HitTid, const FViewCameraState& CameraState,
 		const FRay& Ray, TArray<int32>& IDsOut)
 	{
 		if (!ensure(HitTid != IndexConstants::InvalidID && Mesh.IsTriangle(HitTid)))
@@ -255,11 +344,20 @@ namespace UVEditorMeshSelectionMechanicLocals
 		case ESelectionMode::Island:
 		{
 			TRACE_CPUPROFILER_EVENT_SCOPE(Component);
-
 			FMeshConnectedComponents MeshSelectedComponent(&Mesh);
 			TArray<int32> SeedTriangles;
-			SeedTriangles.Add(HitTid);
-			MeshSelectedComponent.FindTrianglesConnectedToSeeds(SeedTriangles);
+			SeedTriangles.Add(HitTid);			
+			auto UVIslandPredicate = [&UnwrapMesh,&UVOverlay](int32 Triangle0, int32 Triangle1)
+			{
+				if (UnwrapMesh.IsTriangle(Triangle0))
+				{
+					return UnwrapMesh.GetTriNeighbourTris(Triangle0).Contains(Triangle1);
+				}
+				else { // Triangle0 is unset, return true if Triangle1 is also.
+					return !UnwrapMesh.IsTriangle(Triangle1);
+				}
+			};
+			MeshSelectedComponent.FindTrianglesConnectedToSeeds(SeedTriangles, UVIslandPredicate);
 			ensure(MeshSelectedComponent.Components.Num() == 1); // Expect each triangle to only be in a single component
 			IDsOut.Append(MoveTemp(MeshSelectedComponent.Components[0].Indices));
 			break;
@@ -341,6 +439,115 @@ namespace UVEditorMeshSelectionMechanicLocals
 		return !IDsOut.IsEmpty();
 	}
 
+	TArray<int32> Convert3DHitsTo2DHits(ESelectionMode SelectionMode,
+		const FDynamicMesh3& UnwrapMesh, const FDynamicMesh3& AppliedMesh, const FDynamicMeshUVOverlay& UVOverlay, const TArray<int32>& IDsIn, TArray<int32>& AppliedMeshOnlyIDsOut)
+	{
+		TSet<int32> IDsOut;
+		AppliedMeshOnlyIDsOut.Empty();
+
+		auto FindUnwrapEdges = [&AppliedMesh, &UnwrapMesh](int32 Eid, TSet<int32>& OutUnwrapEdges)
+		{
+			bool bFoundEdges = false;
+			FIndex2i Triangles = AppliedMesh.GetEdgeT(Eid);
+			int32 Triangle0EdgeIndex = AppliedMesh.GetTriEdges(Triangles[0]).IndexOf(Eid);
+			int32 Triangle1EdgeIndex = Triangles[1] == IndexConstants::InvalidID ? IndexConstants::InvalidID : AppliedMesh.GetTriEdges(Triangles[1]).IndexOf(Eid);
+			if (UnwrapMesh.IsTriangle(Triangles[0]))
+			{
+				OutUnwrapEdges.Add(UnwrapMesh.GetTriEdge(Triangles[0], Triangle0EdgeIndex));
+				bFoundEdges = true;
+			}
+			if (UnwrapMesh.IsTriangle(Triangles[1]))
+			{
+				OutUnwrapEdges.Add(UnwrapMesh.GetTriEdge(Triangles[1], Triangle1EdgeIndex));
+				bFoundEdges = true;
+			}
+			return bFoundEdges;
+		};
+
+		switch (SelectionMode)
+		{
+		case ESelectionMode::Triangle:			
+		case ESelectionMode::Island:
+		case ESelectionMode::Mesh:
+			for (int32 Tid : IDsIn)
+			{
+				if (UVOverlay.IsSetTriangle(Tid))
+				{
+					IDsOut.Add(Tid);
+				}
+				else
+				{
+					AppliedMeshOnlyIDsOut.Add(Tid);
+				}
+			}
+			break;
+		case ESelectionMode::Edge:
+			for (int32 Eid : IDsIn)
+			{				
+				if(!FindUnwrapEdges(Eid, IDsOut))				
+				{
+					AppliedMeshOnlyIDsOut.Add(Eid);
+				}
+			}
+			break;			
+		case ESelectionMode::Vertex:
+			for (int32 Vid : IDsIn)
+			{
+				TArray<int32> ElementsForVid;
+				UVOverlay.GetVertexElements(Vid, ElementsForVid);
+				IDsOut.Append(ElementsForVid);
+				if (ElementsForVid.IsEmpty())
+				{
+					AppliedMeshOnlyIDsOut.Add(Vid);
+				}
+			}
+			break;
+		default:
+			ensure(false);			
+			break;
+		}
+		return IDsOut.Array();
+	}
+	
+	TArray<int32> Convert2DHitsTo3DHits(ESelectionMode SelectionMode,
+		const FDynamicMesh3& UnwrapMesh, const FDynamicMesh3& AppliedMesh, const FDynamicMeshUVOverlay& UVOverlay, TArray<int32>& IDsIn)
+	{
+		TArray<int32> IDsOut;
+		
+		auto FindAppliedEdge = [&UnwrapMesh, &AppliedMesh, &UVOverlay](int32 Eid)
+		{
+			int32 Vid1, Vid2;
+			FDynamicMesh3::FEdge EdgeInfo = UnwrapMesh.GetEdge(Eid);
+			Vid1 = UVOverlay.GetParentVertex(EdgeInfo.Vert[0]);
+			Vid2 = UVOverlay.GetParentVertex(EdgeInfo.Vert[1]);	
+			return AppliedMesh.FindEdge(Vid1, Vid2);
+		};
+
+		switch (SelectionMode)
+		{
+		case ESelectionMode::Triangle:
+		case ESelectionMode::Island:
+		case ESelectionMode::Mesh:
+			return IDsIn;
+		case ESelectionMode::Edge:
+			for (int32 Eid : IDsIn)
+			{
+				IDsOut.Add(FindAppliedEdge(Eid));
+			}
+			return IDsOut;
+		case ESelectionMode::Vertex:
+			for (int32 Vid : IDsIn)
+			{
+				IDsOut.Add(UVOverlay.GetParentVertex(Vid));
+			}
+			return IDsOut;
+		default:
+			ensure(false);
+			return IDsOut;
+		}
+	}
+
+
 	/**
 	 * Undo/redo transaction for selection mode changes
 	 */
@@ -400,6 +607,34 @@ void UUVEditorMeshSelectionMechanic::Setup(UInteractiveTool* ParentToolIn)
 	EmitChangeAPI = ContextStore->FindContext<UUVToolEmitChangeAPI>();
 	check(EmitChangeAPI);
 
+	UnwrapClickTargetRouter = NewObject<ULocalSingleClickInputBehavior>();
+
+	UnwrapClickTargetRouter->IsHitByClickFunc = [this](const FInputDeviceRay& ClickPos) { return IsHitByClick(ClickPos, false); };
+	UnwrapClickTargetRouter->OnClickedFunc = [this](const FInputDeviceRay& ClickPos) { OnClicked(ClickPos, false); };
+	UnwrapClickTargetRouter->OnUpdateModifierStateFunc = [this](int ModifierID, bool bIsOn) { OnUpdateModifierState(ModifierID,bIsOn); };
+
+	UnwrapHoverBehaviorTargetRouter = NewObject<ULocalMouseHoverBehavior>();
+
+	UnwrapHoverBehaviorTargetRouter->BeginHitTestFunc = [this](const FInputDeviceRay& PressPos) { return BeginHoverSequenceHitTest(PressPos, false); };
+	UnwrapHoverBehaviorTargetRouter->OnBeginHoverFunc = [this](const FInputDeviceRay& DevicePos) { return OnBeginHover(DevicePos); };
+	UnwrapHoverBehaviorTargetRouter->OnUpdateHoverFunc = [this](const FInputDeviceRay& DevicePos) {return OnUpdateHover(DevicePos, false); };
+	UnwrapHoverBehaviorTargetRouter->OnEndHoverFunc = [this]() {return OnEndHover(); };
+
+	LivePreviewClickTargetRouter = NewObject<ULocalSingleClickInputBehavior>();
+
+	LivePreviewClickTargetRouter->IsHitByClickFunc = [this](const FInputDeviceRay& ClickPos) { return IsHitByClick(ClickPos, true); };
+	LivePreviewClickTargetRouter->OnClickedFunc = [this](const FInputDeviceRay& ClickPos) { OnClicked(ClickPos, true); };
+	LivePreviewClickTargetRouter->OnUpdateModifierStateFunc = [this](int ModifierID, bool bIsOn) { OnUpdateModifierState(ModifierID, bIsOn); };
+
+	LivePreviewHoverBehaviorTargetRouter = NewObject<ULocalMouseHoverBehavior>();
+
+	LivePreviewHoverBehaviorTargetRouter->BeginHitTestFunc = [this](const FInputDeviceRay& PressPos) {
+		return BeginHoverSequenceHitTest(PressPos, true);
+	};
+	LivePreviewHoverBehaviorTargetRouter->OnBeginHoverFunc = [this](const FInputDeviceRay& DevicePos) { return OnBeginHover(DevicePos); };
+	LivePreviewHoverBehaviorTargetRouter->OnUpdateHoverFunc = [this](const FInputDeviceRay& DevicePos) {return OnUpdateHover(DevicePos, true); };
+	LivePreviewHoverBehaviorTargetRouter->OnEndHoverFunc = [this]() {return OnEndHover(); };
+
 	// This will be the target for the click drag behavior below
 	MarqueeMechanic = NewObject<URectangleMarqueeMechanic>();
 	MarqueeMechanic->bUseExternalClickDragBehavior = true;
@@ -408,16 +643,19 @@ void UUVEditorMeshSelectionMechanic::Setup(UInteractiveTool* ParentToolIn)
 	// TODO(Performance) :DynamicMarqueeSelection It would be cool to have the marquee selection update dynamically as
 	//  the rectangle gets changed, right now this isn't interactive for large meshes so we disabled it
 	//MarqueeMechanic->OnDragRectangleChanged.AddUObject(this, &UUVEditorMeshSelectionMechanic::OnDragRectangleChanged);
-	MarqueeMechanic->OnDragRectangleFinished.AddUObject(this, &UUVEditorMeshSelectionMechanic::OnDragRectangleFinished);
+	MarqueeMechanic->OnDragRectangleFinished.AddLambda( [this](const FCameraRectangle& Rectangle, bool bCancelled)
+		{
+				OnDragRectangleFinished(Rectangle, bCancelled, false);
+		});
 
 	USingleClickOrDragInputBehavior* ClickOrDragBehavior = NewObject<USingleClickOrDragInputBehavior>();
-	ClickOrDragBehavior->Initialize(this, MarqueeMechanic);
+	ClickOrDragBehavior->Initialize(UnwrapClickTargetRouter, MarqueeMechanic);
 	ClickOrDragBehavior->Modifiers.RegisterModifier(ShiftModifierID, FInputDeviceState::IsShiftKeyDown);
 	ClickOrDragBehavior->Modifiers.RegisterModifier(CtrlModifierID, FInputDeviceState::IsCtrlKeyDown);
 	ParentTool->AddInputBehavior(ClickOrDragBehavior);
 
 	UMouseHoverBehavior* HoverBehavior = NewObject<UMouseHoverBehavior>();
-	HoverBehavior->Initialize(this);
+	HoverBehavior->Initialize(UnwrapHoverBehaviorTargetRouter);
 	ParentTool->AddInputBehavior(HoverBehavior);
 
 	ViewportButtonsAPI = ContextStore->FindContext<UUVToolViewportButtonsAPI>();
@@ -430,9 +668,56 @@ void UUVEditorMeshSelectionMechanic::Setup(UInteractiveTool* ParentToolIn)
 	// Make sure we match the activated button
 	SelectionMode = ViewportButtonsAPI->GetSelectionMode();
 	SetIsEnabled(bIsEnabled);
+
+	LivePreviewAPI = ContextStore->FindContext<UUVToolLivePreviewAPI>();
+	check(LivePreviewAPI);
+
+	// Set things up for being able to add behaviors to live preview.
+	LivePreviewBehaviorSet = NewObject<UInputBehaviorSet>();
+	LivePreviewBehaviorSource = NewObject<ULocalInputBehaviorSource>();
+	LivePreviewBehaviorSource->GetInputBehaviorsFunc = [this]() { return LivePreviewBehaviorSet; };
+
+	// This will be the target for the click drag behavior below
+	LivePreviewMarqueeMechanic = NewObject<URectangleMarqueeMechanic>();
+	LivePreviewMarqueeMechanic->bUseExternalClickDragBehavior = true;
+	LivePreviewMarqueeMechanic->bUseExternalUpdateCameraState = true;
+	LivePreviewMarqueeMechanic->UpdateCameraStateFunc = [this]() {
+		FViewCameraState LivePreviewCameraState;
+		LivePreviewAPI->GetLivePreviewCameraState(LivePreviewCameraState);
+		return LivePreviewCameraState;
+	};
+	LivePreviewMarqueeMechanic->Setup(ParentToolIn);
+	LivePreviewMarqueeMechanic->OnDragRectangleStarted.AddUObject(this, &UUVEditorMeshSelectionMechanic::OnDragRectangleStarted);
+	// TODO(Performance) :DynamicMarqueeSelection It would be cool to have the marquee selection update dynamically as
+	//  the rectangle gets changed, right now this isn't interactive for large meshes so we disabled it
+	//LivePreviewMarqueeMechanic->OnDragRectangleChanged.AddUObject(this, &UUVEditorMeshSelectionMechanic::OnDragRectangleChanged);
+	LivePreviewMarqueeMechanic->OnDragRectangleFinished.AddLambda([this](const FCameraRectangle& Rectangle, bool bCancelled)
+		{
+			OnDragRectangleFinished(Rectangle, bCancelled, true);
+		});
+
+	// Set up click and hover behaviors both for the 2d (unwrap) viewport and the 3d 
+	// (applied/live preview) viewport
+	USingleClickOrDragInputBehavior* LivePreviewClickOrDragBehavior = NewObject<USingleClickOrDragInputBehavior>();
+	LivePreviewClickOrDragBehavior->Initialize(LivePreviewClickTargetRouter, LivePreviewMarqueeMechanic);
+	LivePreviewClickOrDragBehavior->Modifiers.RegisterModifier(ShiftModifierID, FInputDeviceState::IsShiftKeyDown);
+	LivePreviewClickOrDragBehavior->Modifiers.RegisterModifier(CtrlModifierID, FInputDeviceState::IsCtrlKeyDown);
+	LivePreviewBehaviorSet->Add(LivePreviewClickOrDragBehavior, this);
+	
+	UMouseHoverBehavior* LivePreviewHoverBehavior = NewObject<UMouseHoverBehavior>();
+	LivePreviewHoverBehavior->Initialize(LivePreviewHoverBehaviorTargetRouter);
+	LivePreviewBehaviorSet->Add(LivePreviewHoverBehavior, this);
+
+	// Give the live preview behaviors to the live preview input router
+	if (LivePreviewAPI)
+	{
+		LivePreviewInputRouter = LivePreviewAPI->GetLivePreviewInputRouter();
+		LivePreviewInputRouter->RegisterSource(LivePreviewBehaviorSource);
+	}
+
 }
 
-void UUVEditorMeshSelectionMechanic::Initialize(UWorld* World, UUVToolSelectionAPI* SelectionAPIIn)
+void UUVEditorMeshSelectionMechanic::Initialize(UWorld* World, UWorld* LivePreviewWorld, UUVToolSelectionAPI* SelectionAPIIn)
 {
 	// It may be unreasonable to worry about Initialize being called more than once, but let's be safe anyway
 	if (HoverGeometryActor)
@@ -461,6 +746,22 @@ void UUVEditorMeshSelectionMechanic::Initialize(UWorld* World, UUVToolSelectionA
 	HoverLineSet->SetLineMaterial(ToolSetupUtil::GetDefaultLineComponentMaterial(GetParentTool()->GetToolManager(), false));
 	HoverLineSet->AttachToComponent(HoverTriangleSet.Get(), FAttachmentTransformRules::KeepWorldTransform);
 	HoverLineSet->RegisterComponent();
+
+	LivePreviewHoverGeometryActor = LivePreviewWorld->SpawnActor<APreviewGeometryActor>();
+
+	LivePreviewHoverTriangleSet = NewObject<UTriangleSetComponent>(LivePreviewHoverGeometryActor);
+	LivePreviewHoverGeometryActor->SetRootComponent(LivePreviewHoverTriangleSet.Get());
+	LivePreviewHoverTriangleSet->RegisterComponent();
+
+	LivePreviewHoverPointSet = NewObject<UPointSetComponent>(LivePreviewHoverGeometryActor);
+	LivePreviewHoverPointSet->SetPointMaterial(ToolSetupUtil::GetDefaultPointComponentMaterial(GetParentTool()->GetToolManager(), false));
+	LivePreviewHoverPointSet->AttachToComponent(LivePreviewHoverTriangleSet.Get(), FAttachmentTransformRules::KeepWorldTransform);
+	LivePreviewHoverPointSet->RegisterComponent();
+
+	LivePreviewHoverLineSet = NewObject<ULineSetComponent>(LivePreviewHoverGeometryActor);
+	LivePreviewHoverLineSet->SetLineMaterial(ToolSetupUtil::GetDefaultLineComponentMaterial(GetParentTool()->GetToolManager(), false));
+	LivePreviewHoverLineSet->AttachToComponent(LivePreviewHoverTriangleSet.Get(), FAttachmentTransformRules::KeepWorldTransform);
+	LivePreviewHoverLineSet->RegisterComponent();
 }
 
 void UUVEditorMeshSelectionMechanic::SetIsEnabled(bool bIsEnabledIn)
@@ -470,6 +771,12 @@ void UUVEditorMeshSelectionMechanic::SetIsEnabled(bool bIsEnabledIn)
 	{
 		MarqueeMechanic->SetIsEnabled(bIsEnabled && SelectionMode != ESelectionMode::None);
 	}
+	if (LivePreviewMarqueeMechanic)
+	{
+		bool bLiveViewportMarqueeActive = bIsEnabled;
+		bLiveViewportMarqueeActive &= SelectionMode != ESelectionMode::None;		
+		LivePreviewMarqueeMechanic->SetIsEnabled(bLiveViewportMarqueeActive);
+	}
 	if (ViewportButtonsAPI)
 	{
 		ViewportButtonsAPI->SetSelectionButtonsEnabled(bIsEnabledIn);
@@ -484,6 +791,10 @@ void UUVEditorMeshSelectionMechanic::SetShowHoveredElements(bool bShow)
 		HoverPointSet->Clear();
 		HoverLineSet->Clear();
 		HoverTriangleSet->Clear();
+
+		LivePreviewHoverPointSet->Clear();
+		LivePreviewHoverLineSet->Clear();
+		LivePreviewHoverTriangleSet->Clear();
 	}
 }
 
@@ -494,11 +805,23 @@ void UUVEditorMeshSelectionMechanic::Shutdown()
 		HoverGeometryActor->Destroy();
 		HoverGeometryActor = nullptr;
 	}
+	if (LivePreviewHoverGeometryActor)
+	{
+		LivePreviewHoverGeometryActor->Destroy();
+		LivePreviewHoverGeometryActor = nullptr;
+	}
+	if (LivePreviewAPI)
+	{
+		LivePreviewInputRouter->DeregisterSource(LivePreviewBehaviorSource);
+	}
+
 	SelectionAPI = nullptr;
 	ViewportButtonsAPI = nullptr;
+	LivePreviewAPI = nullptr;
 	EmitChangeAPI = nullptr;
 	MarqueeMechanic = nullptr;
 	HoverTriangleSetMaterial = nullptr;
+	LivePreviewMarqueeMechanic = nullptr;
 }
 
 void UUVEditorMeshSelectionMechanic::SetTargets(const TArray<TObjectPtr<UUVEditorToolMeshInput>>& TargetsIn)
@@ -516,23 +839,47 @@ void UUVEditorMeshSelectionMechanic::SetTargets(const TArray<TObjectPtr<UUVEdito
 
 	// Get or create spatials
 	// Initialize the AABB trees from cached values, or make new ones
-	MeshSpatials.Reset();
+	UnwrapMeshSpatials.Reset();
+	AppliedMeshSpatials.Reset();
 	for (TObjectPtr<UUVEditorToolMeshInput> Target : Targets)
 	{
-		TSharedPtr<FDynamicMeshAABBTree3> Tree = TreeStore->Get(Target->UnwrapCanonical.Get());
-		if (!Tree)
+		TSharedPtr<FDynamicMeshAABBTree3> UnwrapTree = TreeStore->Get(Target->UnwrapCanonical.Get());
+		if (!UnwrapTree)
 		{
 			TRACE_CPUPROFILER_EVENT_SCOPE(BuildAABBTreeForTarget);
-			Tree = MakeShared<FDynamicMeshAABBTree3>();
-			Tree->SetMesh(Target->UnwrapCanonical.Get(), false);
+			UnwrapTree = MakeShared<FDynamicMeshAABBTree3>();
+			UnwrapTree->SetMesh(Target->UnwrapCanonical.Get(), false);
 			// For now we split round-robin on the X/Y axes TODO Experiment with better splitting heuristics
 			FDynamicMeshAABBTree3::GetSplitAxisFunc GetSplitAxis = [](int Depth, const FAxisAlignedBox3d&) { return Depth % 2; };
 			// Note: 16 tris/leaf was chosen with data collected by SpatialBenchmarks.cpp in GeometryProcessingUnitTests
-			Tree->SetBuildOptions(16, MoveTemp(GetSplitAxis));
-			Tree->Build();
-			TreeStore->Set(Target->UnwrapCanonical.Get(), Tree, Target);
+			UnwrapTree->SetBuildOptions(16, MoveTemp(GetSplitAxis));
+			UnwrapTree->Build();
+			TreeStore->Set(Target->UnwrapCanonical.Get(), UnwrapTree, Target);
+		}
+		UnwrapMeshSpatials.Add(UnwrapTree);
+		TSharedPtr<FDynamicMeshAABBTree3> AppliedTree = TreeStore->Get(Target->AppliedCanonical.Get());
+		if (!AppliedTree)
+		{
+			TRACE_CPUPROFILER_EVENT_SCOPE(BuildAABBTreeForTarget);
+			AppliedTree = MakeShared<FDynamicMeshAABBTree3>();
+			AppliedTree->SetMesh(Target->AppliedCanonical.Get(), false);
+			AppliedTree->Build();
+			TreeStore->Set(Target->AppliedCanonical.Get(), AppliedTree, Target);
 		}
-		MeshSpatials.Add(Tree);
+		AppliedMeshSpatials.Add(AppliedTree);
+	}
+
+}
+
+TSharedPtr<FDynamicMeshAABBTree3> UUVEditorMeshSelectionMechanic::GetMeshSpatial(int32 TargetId, bool bUseUnwrap)
+{
+	if (bUseUnwrap)
+	{
+		return UnwrapMeshSpatials[TargetId];
+	}
+	else
+	{
+		return AppliedMeshSpatials[TargetId];
 	}
 }
 
@@ -549,9 +896,20 @@ void UUVEditorMeshSelectionMechanic::DrawHUD(FCanvas* Canvas, IToolsContextRende
 	MarqueeMechanic->DrawHUD(Canvas, RenderAPI);
 }
 
-FInputRayHit UUVEditorMeshSelectionMechanic::IsHitByClick(const FInputDeviceRay& ClickPos)
+void UUVEditorMeshSelectionMechanic::LivePreviewRender(IToolsContextRenderAPI* RenderAPI)
+{
+	LivePreviewMarqueeMechanic->Render(RenderAPI);
+}
+
+void UUVEditorMeshSelectionMechanic::LivePreviewDrawHUD(FCanvas* Canvas, IToolsContextRenderAPI* RenderAPI)
 {
-	FInputRayHit Hit; 
+	LivePreviewMarqueeMechanic->DrawHUD(Canvas, RenderAPI);
+}
+
+FInputRayHit UUVEditorMeshSelectionMechanic::IsHitByClick(const FInputDeviceRay& ClickPos, bool bSourceIsLivePreview)
+{
+	FInputRayHit Hit;
+
 	// If enabled, return a hit so we always capture and can clear the selection
 	Hit.bHit = bIsEnabled && SelectionMode != ESelectionMode::None;
 	return Hit;	
@@ -600,7 +958,8 @@ void UUVEditorMeshSelectionMechanic::SetSelectionMode(
 	// See whether a conversion is not necessary
 	FUVToolSelection::EType ExpectedSelectionType = ToCompatibleDynamicMeshSelectionType(SelectionMode);
 	FUVToolSelection::EType CurrentSelectionType = SelectionAPI->GetSelectionsType();
-	if (!SelectionAPI->HaveSelections() || ExpectedSelectionType == CurrentSelectionType 
+	if ((!SelectionAPI->HaveSelections() && !SelectionAPI->HaveUnsetElementAppliedMeshSelections())
+		|| ExpectedSelectionType == CurrentSelectionType
 		|| !Options.bConvertExisting || SelectionMode == ESelectionMode::None)
 	{
 		// No conversion needed
@@ -613,188 +972,202 @@ void UUVEditorMeshSelectionMechanic::SetSelectionMode(
 	
 	// We're going to convert the existing selection.
 	const TArray<FUVToolSelection>& OriginalSelections = SelectionAPI->GetSelections();
-	TArray<FUVToolSelection> NewSelections;
+	const TArray<FUVToolSelection>& OriginalUnsetSelections = SelectionAPI->GetUnsetElementAppliedMeshSelections();
 
-	for (const FUVToolSelection& OriginalSelection : OriginalSelections)
+	auto ConvertSelections = [](bool bUnsetSelections, const TArray<FUVToolSelection>& OriginalSelections, FUVToolSelection::EType CurrentSelectionType, FUVToolSelection::EType ExpectedSelectionType)
 	{
-		FDynamicMesh3* Mesh = OriginalSelection.Target->UnwrapCanonical.Get();
-		
-		FUVToolSelection NewSelection;
-		NewSelection.Target = OriginalSelection.Target;
-		NewSelection.Type = ExpectedSelectionType;
+		TArray<FUVToolSelection> NewSelections;
+		for (const FUVToolSelection& OriginalSelection : OriginalSelections)
+		{
+			FDynamicMesh3* Mesh = bUnsetSelections ? OriginalSelection.Target->AppliedCanonical.Get(): OriginalSelection.Target->UnwrapCanonical.Get();
 
-		const TSet<int32> OriginalIDs = OriginalSelection.SelectedIDs;
-		TSet<int32>& NewIDs = NewSelection.SelectedIDs;
+			FUVToolSelection NewSelection;
+			check(CurrentSelectionType == OriginalSelection.Type);
+			NewSelection.Target = OriginalSelection.Target;
+			NewSelection.Type = ExpectedSelectionType;
 
-		auto VerticesToEdges = [Mesh, &OriginalIDs, &NewIDs]()
-		{
-			TRACE_CPUPROFILER_EVENT_SCOPE(MeshSelectionMechanic_ChangeSelectionMode_VerticesToEdges);
+			const TSet<int32> OriginalIDs = OriginalSelection.SelectedIDs;
+			TSet<int32>& NewIDs = NewSelection.SelectedIDs;
 
-			for (int32 Vid : OriginalIDs)
+			auto VerticesToEdges = [Mesh, &OriginalIDs, &NewIDs]()
 			{
-				for (int32 Eid : Mesh->VtxEdgesItr(Vid))
+				TRACE_CPUPROFILER_EVENT_SCOPE(MeshSelectionMechanic_ChangeSelectionMode_VerticesToEdges);
+
+				for (int32 Vid : OriginalIDs)
 				{
-					if (!NewIDs.Contains(Eid))
+					for (int32 Eid : Mesh->VtxEdgesItr(Vid))
 					{
-						FIndex2i Verts = Mesh->GetEdgeV(Eid);
-						if (OriginalIDs.Contains(Verts.A) &&
-							OriginalIDs.Contains(Verts.B))
+						if (!NewIDs.Contains(Eid))
 						{
-							NewIDs.Add(Eid);
+							FIndex2i Verts = Mesh->GetEdgeV(Eid);
+							if (OriginalIDs.Contains(Verts.A) &&
+								OriginalIDs.Contains(Verts.B))
+							{
+								NewIDs.Add(Eid);
+							}
 						}
 					}
 				}
-			}
-		};
-
-		auto VerticesToTriangles = [Mesh, &OriginalIDs, &NewIDs]()
-		{
-			TRACE_CPUPROFILER_EVENT_SCOPE(MeshSelectionMechanic_ChangeSelectionMode_VerticesToTriangles);
+			};
 
-			for (int32 Vid : OriginalIDs)
+			auto VerticesToTriangles = [Mesh, &OriginalIDs, &NewIDs]()
 			{
-				for (int32 Tid : Mesh->VtxTrianglesItr(Vid))
+				TRACE_CPUPROFILER_EVENT_SCOPE(MeshSelectionMechanic_ChangeSelectionMode_VerticesToTriangles);
+
+				for (int32 Vid : OriginalIDs)
 				{
-					if (!NewIDs.Contains(Tid))
+					for (int32 Tid : Mesh->VtxTrianglesItr(Vid))
 					{
-						FIndex3i Verts = Mesh->GetTriangle(Tid);
-						if (OriginalIDs.Contains(Verts.A) &&
-							OriginalIDs.Contains(Verts.B) &&
-							OriginalIDs.Contains(Verts.C))
+						if (!NewIDs.Contains(Tid))
 						{
-							NewIDs.Add(Tid);
+							FIndex3i Verts = Mesh->GetTriangle(Tid);
+							if (OriginalIDs.Contains(Verts.A) &&
+								OriginalIDs.Contains(Verts.B) &&
+								OriginalIDs.Contains(Verts.C))
+							{
+								NewIDs.Add(Tid);
+							}
 						}
 					}
 				}
-			}
-		};
+			};
 
-		auto EdgesToVertices = [Mesh, &OriginalIDs, &NewIDs]()
-		{
-			TRACE_CPUPROFILER_EVENT_SCOPE(MeshSelectionMechanic_ChangeSelectionMode_EdgesToVertices);
-
-			for (int32 Eid : OriginalIDs)
+			auto EdgesToVertices = [Mesh, &OriginalIDs, &NewIDs]()
 			{
-				FIndex2i Verts = Mesh->GetEdgeV(Eid);
-				NewIDs.Add(Verts.A);
-				NewIDs.Add(Verts.B);
-			}
-		};
-
-		// Triangles with two selected edges will be selected
-		auto EdgesToTriangles = [Mesh, &OriginalIDs, &NewIDs]()
-		{
-			TRACE_CPUPROFILER_EVENT_SCOPE(MeshSelectionMechanic_ChangeSelectionMode_EdgesToTriangles);
+				TRACE_CPUPROFILER_EVENT_SCOPE(MeshSelectionMechanic_ChangeSelectionMode_EdgesToVertices);
 
-			TArray<int32> FoundTriangles;
-			for (int32 Eid : OriginalIDs)
-			{
-				FIndex2i Tris = Mesh->GetEdgeT(Eid);
-				FoundTriangles.Add(Tris.A);
-				if (Tris.B != IndexConstants::InvalidID)
+				for (int32 Eid : OriginalIDs)
 				{
-					FoundTriangles.Add(Tris.B);
+					FIndex2i Verts = Mesh->GetEdgeV(Eid);
+					NewIDs.Add(Verts.A);
+					NewIDs.Add(Verts.B);
 				}
-			}
+			};
 
-			if (FoundTriangles.Num() < 2)
+			// Triangles with two selected edges will be selected
+			auto EdgesToTriangles = [Mesh, &OriginalIDs, &NewIDs]()
 			{
-				return;
-			}
+				TRACE_CPUPROFILER_EVENT_SCOPE(MeshSelectionMechanic_ChangeSelectionMode_EdgesToTriangles);
 
-			Algo::Sort(FoundTriangles);
+				TArray<int32> FoundTriangles;
+				for (int32 Eid : OriginalIDs)
+				{
+					FIndex2i Tris = Mesh->GetEdgeT(Eid);
+					FoundTriangles.Add(Tris.A);
+					if (Tris.B != IndexConstants::InvalidID)
+					{
+						FoundTriangles.Add(Tris.B);
+					}
+				}
 
-			for (int I = 0; I < FoundTriangles.Num() - 1; I++)
-			{
-				if (FoundTriangles[I] == FoundTriangles[I + 1])
+				if (FoundTriangles.Num() < 2)
 				{
-					NewIDs.Add(FoundTriangles[I]);
-					I++;
+					return;
 				}
-			}
-		};
 
-		auto TrianglesToVertices = [Mesh, &OriginalIDs, &NewIDs]()
-		{
-			TRACE_CPUPROFILER_EVENT_SCOPE(MeshSelectionMechanic_ChangeSelectionMode_TrianglesToVertices);
+				Algo::Sort(FoundTriangles);
+
+				for (int I = 0; I < FoundTriangles.Num() - 1; I++)
+				{
+					if (FoundTriangles[I] == FoundTriangles[I + 1])
+					{
+						NewIDs.Add(FoundTriangles[I]);
+						I++;
+					}
+				}
+			};
 
-			for (int32 Tid : OriginalIDs)
+			auto TrianglesToVertices = [Mesh, &OriginalIDs, &NewIDs]()
 			{
-				FIndex3i Verts = Mesh->GetTriangle(Tid);
-				NewIDs.Add(Verts.A);
-				NewIDs.Add(Verts.B);
-				NewIDs.Add(Verts.C);
-			}
-		};
+				TRACE_CPUPROFILER_EVENT_SCOPE(MeshSelectionMechanic_ChangeSelectionMode_TrianglesToVertices);
 
-		auto TrianglesToEdges = [Mesh, &OriginalIDs, &NewIDs]()
-		{
-			TRACE_CPUPROFILER_EVENT_SCOPE(MeshSelectionMechanic_ChangeSelectionMode_TrianglesToEdges);
+				for (int32 Tid : OriginalIDs)
+				{
+					FIndex3i Verts = Mesh->GetTriangle(Tid);
+					NewIDs.Add(Verts.A);
+					NewIDs.Add(Verts.B);
+					NewIDs.Add(Verts.C);
+				}
+			};
 
-			for (int32 Tid : OriginalIDs)
+			auto TrianglesToEdges = [Mesh, &OriginalIDs, &NewIDs]()
 			{
-				FIndex3i Edges = Mesh->GetTriEdgesRef(Tid);
-				NewIDs.Add(Edges.A);
-				NewIDs.Add(Edges.B);
-				NewIDs.Add(Edges.C);
-			}
-		};
+				TRACE_CPUPROFILER_EVENT_SCOPE(MeshSelectionMechanic_ChangeSelectionMode_TrianglesToEdges);
 
-		switch (CurrentSelectionType)
-		{
-		case FUVToolSelection::EType::Vertex:
-			switch (ExpectedSelectionType)
-			{
-			case FUVToolSelection::EType::Vertex:
-				ensure(false); // Should have been an early-out
-				break;
-			case FUVToolSelection::EType::Edge:
-				VerticesToEdges();
-				break;
-			case FUVToolSelection::EType::Triangle:
-				VerticesToTriangles();
-				break;
-			}
-			break; // SelectionMode
-		case FUVToolSelection::EType::Edge:
-			switch (ExpectedSelectionType)
+				for (int32 Tid : OriginalIDs)
+				{
+					FIndex3i Edges = Mesh->GetTriEdgesRef(Tid);
+					NewIDs.Add(Edges.A);
+					NewIDs.Add(Edges.B);
+					NewIDs.Add(Edges.C);
+				}
+			};
+
+			switch (CurrentSelectionType)
 			{
 			case FUVToolSelection::EType::Vertex:
-				EdgesToVertices();
-				break;
+				switch (ExpectedSelectionType)
+				{
+				case FUVToolSelection::EType::Vertex:
+					ensure(false); // Should have been an early-out
+					break;
+				case FUVToolSelection::EType::Edge:
+					VerticesToEdges();
+					break;
+				case FUVToolSelection::EType::Triangle:
+					VerticesToTriangles();
+					break;
+				}
+				break; // SelectionMode
 			case FUVToolSelection::EType::Edge:
-				ensure(false); // Should have been an early-out
-				break;
+				switch (ExpectedSelectionType)
+				{
+				case FUVToolSelection::EType::Vertex:
+					EdgesToVertices();
+					break;
+				case FUVToolSelection::EType::Edge:
+					ensure(false); // Should have been an early-out
+					break;
+				case FUVToolSelection::EType::Triangle:
+					EdgesToTriangles();
+					break;
+				}
+				break; // SelectionMode
 			case FUVToolSelection::EType::Triangle:
-				EdgesToTriangles();
-				break;
+				switch (ExpectedSelectionType)
+				{
+				case FUVToolSelection::EType::Vertex:
+					TrianglesToVertices();
+					break;
+				case FUVToolSelection::EType::Edge:
+					TrianglesToEdges();
+					break;
+				case FUVToolSelection::EType::Triangle:
+					ensure(false); // Should have been an early-out
+					break;
+				}
+				break; // SelectionMode
 			}
-			break; // SelectionMode
-		case FUVToolSelection::EType::Triangle:
-			switch (ExpectedSelectionType)
+
+			if (!NewSelection.IsEmpty())
 			{
-			case FUVToolSelection::EType::Vertex:
-				TrianglesToVertices();
-				break;
-			case FUVToolSelection::EType::Edge:
-				TrianglesToEdges();
-				break;
-			case FUVToolSelection::EType::Triangle:
-				ensure(false); // Should have been an early-out
-				break;
+				NewSelections.Add(MoveTemp(NewSelection));
 			}
-			break; // SelectionMode
 		}
 
-		if (!NewSelection.IsEmpty())
-		{
-			NewSelections.Add(MoveTemp(NewSelection));
-		}
-	}
+		return NewSelections;
+	};
+
+	TArray<FUVToolSelection> NewSelections = ConvertSelections(false, OriginalSelections,CurrentSelectionType, ExpectedSelectionType);
+	TArray<FUVToolSelection> NewUnsetSelections = ConvertSelections(true, OriginalUnsetSelections, CurrentSelectionType, ExpectedSelectionType);
 
 	// Apply selection change
-	SelectionAPI->SetSelections(NewSelections, Options.bBroadcastIfConverted, 
-		Options.bEmitChanges);
+	SelectionAPI->BeginChange();
+	SelectionAPI->ClearSelections(false, false);
+	SelectionAPI->ClearUnsetElementAppliedMeshSelections(false, false);
+	SelectionAPI->SetSelections(NewSelections, false, false);
+	SelectionAPI->SetUnsetElementAppliedMeshSelections(NewUnsetSelections, false, false);
+	SelectionAPI->EndChangeAndEmitIfModified(Options.bBroadcastIfConverted);
 
 	if (Options.bEmitChanges)
 	{
@@ -833,7 +1206,7 @@ void UUVEditorMeshSelectionMechanic::ModifyExistingSelection(TSet<int32>& Select
 	}
 }
 
-void UUVEditorMeshSelectionMechanic::OnClicked(const FInputDeviceRay& ClickPos)
+void UUVEditorMeshSelectionMechanic::OnClicked(const FInputDeviceRay& ClickPos, bool bSourceIsLivePreview)
 {
 	TRACE_CPUPROFILER_EVENT_SCOPE(MeshSelectionMechanic_OnClicked);
 	
@@ -850,67 +1223,97 @@ void UUVEditorMeshSelectionMechanic::OnClicked(const FInputDeviceRay& ClickPos)
 	int32 HitAssetID = IndexConstants::InvalidID;
 	int32 HitTid = IndexConstants::InvalidID;
 	int32 ExistingSelectionIndex = IndexConstants::InvalidID;
-	TArray<int32> NewIDs;
+	TArray<int32> NewIDs, NewUnsetIDs;
+
+	bool bTidIsHit = GetHitTid(ClickPos, HitTid, HitAssetID, !bSourceIsLivePreview, &ExistingSelectionIndex);
+	bool bHitConversionSucceeded = false;
+	if (bTidIsHit)
+	{
+		FDynamicMesh3* UnwrapMesh = Targets[HitAssetID]->UnwrapCanonical.Get();
+		FDynamicMesh3* AppliedMesh = Targets[HitAssetID]->AppliedCanonical.Get();
+		FDynamicMesh3* Mesh = bSourceIsLivePreview ? AppliedMesh : UnwrapMesh;
+		FDynamicMeshUVOverlay& UVOverlay = *Mesh->Attributes()->GetUVLayer(Targets[HitAssetID]->UVLayerIndex);
+		bHitConversionSucceeded = ConvertToHitElementList(SelectionMode, *Mesh, *UnwrapMesh, UVOverlay, HitTid, CameraState, ClickPos.WorldRay, NewIDs);
 
-	// Do the raycast and get selected elements
-	if (!GetHitTid(ClickPos, HitTid, HitAssetID, &ExistingSelectionIndex)
-		|| !ConvertToHitElementList(SelectionMode, *Targets[HitAssetID]->UnwrapCanonical,
-				HitTid, CameraState, ClickPos.WorldRay, NewIDs))
+		if (bSourceIsLivePreview)
+		{
+			NewIDs = Convert3DHitsTo2DHits(SelectionMode, *UnwrapMesh, *AppliedMesh, UVOverlay, NewIDs, NewUnsetIDs);
+		}
+	}
+	if(!bTidIsHit || !bHitConversionSucceeded)
 	{
 		// Failed to select an element. See if selection needs clearing, and exit.
-		if (ShouldRestartSelection() && SelectionAPI->HaveSelections())
+		if (ShouldRestartSelection())
 		{
-			SelectionAPI->ClearSelections(true, true); // broadcast and emit
+			SelectionAPI->BeginChange();
+			if (SelectionAPI->HaveSelections())
+			{
+				SelectionAPI->ClearSelections(false, false);
+			}
+			if (SelectionAPI->HaveUnsetElementAppliedMeshSelections())
+			{
+				SelectionAPI->ClearUnsetElementAppliedMeshSelections(false, false);
+			}
+			SelectionAPI->EndChangeAndEmitIfModified(true); // broadcast and emit
 		}
 		return;
 	}
 
 	TArray<FUVToolSelection> NewSelections;
+	TArray<FUVToolSelection> NewUnsetSelections;
 	if (!ShouldRestartSelection())
 	{
 		NewSelections = SelectionAPI->GetSelections();
+		NewUnsetSelections = SelectionAPI->GetUnsetElementAppliedMeshSelections();
 	}
 
-	if (NewIDs.IsEmpty())
-	{
-		// Nothing to add or modify.
-	}
-	else if (ShouldRestartSelection() 
-		|| (ExistingSelectionIndex == IndexConstants::InvalidID && !ShouldRemoveFromSelection()))
+	auto ProcessNewSelections = [this, &ExistingSelectionIndex, &HitAssetID, &ElementType](const TArray<int32>& IdsIn, TArray<FUVToolSelection>& Selections)
 	{
-		// Make a new selection object
-		NewSelections.Emplace();
-		NewSelections.Last().Target = Targets[HitAssetID];
-		NewSelections.Last().Type = ElementType;
-		NewSelections.Last().SelectedIDs.Append(NewIDs);
-	}
-	else if (ExistingSelectionIndex != IndexConstants::InvalidID)
-	{
-		// Modify the existing selection object
-		ModifyExistingSelection(NewSelections[ExistingSelectionIndex].SelectedIDs, NewIDs);
-
-		// Object may end up empty due to subtraction or toggle, in which case it needs to be removed.
-		if (NewSelections[ExistingSelectionIndex].IsEmpty())
+		if (IdsIn.IsEmpty())
 		{
-			NewSelections.RemoveAt(ExistingSelectionIndex);
+			// Nothing to add or modify.
 		}
-	}
-	else
-	{
-		// The only way we can get here is if didn't have an existing selection and were trying
-		// to remove selection, in which case we do nothing.
-		ensure(ExistingSelectionIndex == IndexConstants::InvalidID && ShouldRemoveFromSelection());
-	}
+		else if (ShouldRestartSelection()
+			|| (ExistingSelectionIndex == IndexConstants::InvalidID && !ShouldRemoveFromSelection()))
+		{
+			// Make a new selection object
+			Selections.Emplace();
+			Selections.Last().Target = Targets[HitAssetID];
+			Selections.Last().Type = ElementType;
+			Selections.Last().SelectedIDs.Append(IdsIn);
+		}
+		else if (ExistingSelectionIndex != IndexConstants::InvalidID)
+		{
+			// Modify the existing selection object
+			ModifyExistingSelection(Selections[ExistingSelectionIndex].SelectedIDs, IdsIn);
 
-	SelectionAPI->SetSelections(NewSelections, true, true); // broadcast and emit
+			// Object may end up empty due to subtraction or toggle, in which case it needs to be removed.
+			if (Selections[ExistingSelectionIndex].IsEmpty())
+			{
+				Selections.RemoveAt(ExistingSelectionIndex);
+			}
+		}
+		else
+		{
+			// The only way we can get here is if didn't have an existing selection and were trying
+			// to remove selection, in which case we do nothing.
+			ensure(ExistingSelectionIndex == IndexConstants::InvalidID && ShouldRemoveFromSelection());
+		}
+	};
+	ProcessNewSelections(NewIDs, NewSelections);
+	ProcessNewSelections(NewUnsetIDs, NewUnsetSelections);
+	SelectionAPI->BeginChange();
+	SelectionAPI->SetSelections(NewSelections, false, false);
+	SelectionAPI->SetUnsetElementAppliedMeshSelections(NewUnsetSelections, false, false);
+	SelectionAPI->EndChangeAndEmitIfModified(true); // broadcast and emit
 }
 
 bool UUVEditorMeshSelectionMechanic::GetHitTid(const FInputDeviceRay& ClickPos, 
-	int32& TidOut, int32& AssetIDOut, int32* ExistingSelectionObjectIndexOut)
+	int32& TidOut, int32& AssetIDOut, bool bUseUnwrap, int32* ExistingSelectionObjectIndexOut)
 {
-	auto RayCastSpatial = [this, &ClickPos, &TidOut, &AssetIDOut](int32 AssetID) {
+	auto RayCastSpatial = [this, &ClickPos, &TidOut, &AssetIDOut, &bUseUnwrap](int32 AssetID) {
 		double RayT = 0;
-		if (MeshSpatials[AssetID]->FindNearestHitTriangle(ClickPos.WorldRay, RayT, TidOut))
+		if (GetMeshSpatial(AssetID, bUseUnwrap)->FindNearestHitTriangle(ClickPos.WorldRay, RayT, TidOut))
 		{
 			AssetIDOut = AssetID;
 			return true;
@@ -920,12 +1323,12 @@ bool UUVEditorMeshSelectionMechanic::GetHitTid(const FInputDeviceRay& ClickPos,
 
 	// Try raycasting the selected meshes first
 	TArray<bool> SpatialTriedFlags;
-	SpatialTriedFlags.SetNum(MeshSpatials.Num());
+	SpatialTriedFlags.SetNum(Targets.Num());
 	const TArray<FUVToolSelection>& Selections = SelectionAPI->GetSelections();
 	for (int32 SelectionIndex = 0; SelectionIndex < Selections.Num(); ++SelectionIndex)
 	{
 		const FUVToolSelection& Selection = Selections[SelectionIndex];
-		if (ensure(Selection.Target.IsValid() && Selection.Target->AssetID < MeshSpatials.Num()))
+		if (ensure(Selection.Target.IsValid() && Selection.Target->AssetID < Targets.Num()))
 		{
 			if (RayCastSpatial(Selection.Target->AssetID))
 			{
@@ -945,7 +1348,7 @@ bool UUVEditorMeshSelectionMechanic::GetHitTid(const FInputDeviceRay& ClickPos,
 	}
 
 	// Try raycasting the other meshes
-	for (int32 AssetID = 0; AssetID < MeshSpatials.Num(); ++AssetID)
+	for (int32 AssetID = 0; AssetID < Targets.Num(); ++AssetID)
 	{
 		if (SpatialTriedFlags[AssetID])
 		{
@@ -967,11 +1370,15 @@ void UUVEditorMeshSelectionMechanic::OnDragRectangleStarted()
 	TRACE_CPUPROFILER_EVENT_SCOPE(MeshSelectionMechanic_OnDragRectangleStarted); // Mark start of drag sequence
 	
 	PreDragSelections = SelectionAPI->GetSelections();
+	PreDragUnsetSelections = SelectionAPI->GetUnsetElementAppliedMeshSelections();
 	SelectionAPI->BeginChange();
 
 	AssetIDToPreDragSelection.Reset();
+	AssetIDToPreDragUnsetSelection.Reset();
 	AssetIDToPreDragSelection.SetNumZeroed(Targets.Num());
+	AssetIDToPreDragUnsetSelection.SetNumZeroed(Targets.Num());
 	FUVToolSelection::EType ExpectedSelectionType = ToCompatibleDynamicMeshSelectionType(SelectionMode);
+
 	if (SelectionAPI->HaveSelections() 
 		&& SelectionAPI->GetSelectionsType() == ExpectedSelectionType)
 	{
@@ -983,57 +1390,140 @@ void UUVEditorMeshSelectionMechanic::OnDragRectangleStarted()
 			}
 		}
 	}
+	if (SelectionAPI->HaveUnsetElementAppliedMeshSelections()
+		&& SelectionAPI->GetSelectionsType() == ExpectedSelectionType)
+	{
+		for (FUVToolSelection& Selection : PreDragUnsetSelections)
+		{
+			if (ensure(Selection.Type == ExpectedSelectionType))
+			{
+				AssetIDToPreDragUnsetSelection[Selection.Target->AssetID] = &Selection;
+			}
+		}
+	}
+
 }
 
-void UUVEditorMeshSelectionMechanic::OnDragRectangleChanged(const FCameraRectangle& CurrentRectangle)
+void UUVEditorMeshSelectionMechanic::OnDragRectangleChanged(const FCameraRectangle& CurrentRectangle, bool bSourceIsLivePreview)
 {
 	TRACE_CPUPROFILER_EVENT_SCOPE(MeshSelectionMechanic_OnDragRectangleChanged);
 
 	using namespace UVEditorMeshSelectionMechanicLocals;
 
-	FAxisAlignedBox2d RectangleXY = GetRectangleXY(CurrentRectangle);
-	TArray<FUVToolSelection> NewSelections;
+	TArray<FUVToolSelection> NewSelections, NewUnsetSelections;
 	FUVToolSelection::EType SelectionType = ToCompatibleDynamicMeshSelectionType(SelectionMode);
 
 	// Gather IDs in each target
 	for (int32 AssetID = 0; AssetID < Targets.Num(); ++AssetID)
 	{
-		TArray<int32> RectangleSelectedIDs;
-		const FDynamicMeshAABBTree3& Tree = *MeshSpatials[AssetID];
+		// Used for Unwrap Selection
+		FAxisAlignedBox2d RectangleXY = GetRectangleXY(CurrentRectangle);
+
+		// Used for Live Preview Selection
+		// TODO: Frustum transformation is done the same way in GroupTopologySelector.cpp and
+		// FractureEditorMode.cpp- should this be placed somewhere common? Not sure where to put it though.
+		FTransform3d TargetTransform = Targets[AssetID]->AppliedPreview->PreviewMesh->GetTransform();
+		FConvexVolume WorldSpaceFrustum = CurrentRectangle.FrustumAsConvexVolume();
+		FMatrix InverseTargetTransform(FTransform(TargetTransform).ToInverseMatrixWithScale());
+		FConvexVolume LocalFrustum;
+		LocalFrustum.Planes.Empty(6);
+		for (const FPlane& Plane : WorldSpaceFrustum.Planes)
+		{
+			LocalFrustum.Planes.Add(Plane.TransformBy(InverseTargetTransform));
+		}
+		LocalFrustum.Init();
+
+		const FDynamicMesh3& Mesh = *Targets[AssetID]->AppliedCanonical;
+		const FDynamicMesh3& UnwrapMesh = *Targets[AssetID]->UnwrapCanonical;
+		const FDynamicMeshUVOverlay& UVOverlay = *Mesh.Attributes()->GetUVLayer(Targets[AssetID]->UVLayerIndex);
+
+		TArray<int32> RectangleSelectedIDs, UnsetRectangleSelectedIDs;
+		const FDynamicMeshAABBTree3& Tree = *GetMeshSpatial(AssetID, !bSourceIsLivePreview);
 
 		if (SelectionMode == ESelectionMode::Vertex)
 		{
 			TRACE_CPUPROFILER_EVENT_SCOPE(MeshSelectionMechanic_OnDragRectangleChanged_Vertex);
 
-			RectangleSelectedIDs = FindAllIntersectionsAxisAlignedBox2(
-				Tree, RectangleXY, AppendVertexIDs, AppendVertexIDsIfIntersected);
+			if (bSourceIsLivePreview)
+			{
+				RectangleSelectedIDs = FindAllIntersections(
+					Tree, LocalFrustum, AppendVertexIDs, AppendVertexIDsIfIntersectedFrustum);
+			}
+			else
+			{
+				RectangleSelectedIDs = FindAllIntersections(
+					Tree, RectangleXY, AppendVertexIDs, AppendVertexIDsIfIntersectedRectangle);
+			}
 		}
 		else if (SelectionMode == ESelectionMode::Edge)
 		{
 			TRACE_CPUPROFILER_EVENT_SCOPE(MeshSelectionMechanic_OnDragRectangleChanged_Edge);
 
-			RectangleSelectedIDs = FindAllIntersectionsAxisAlignedBox2(
-				Tree, RectangleXY, AppendEdgeIDs, AppendEdgeIDsIfIntersected);
+			if (bSourceIsLivePreview)
+			{
+				RectangleSelectedIDs = FindAllIntersections(
+					Tree, LocalFrustum, AppendEdgeIDs, AppendEdgeIDsIfIntersectedFrustum);
+			}
+			else
+			{
+				RectangleSelectedIDs = FindAllIntersections(
+					Tree, RectangleXY, AppendEdgeIDs, AppendEdgeIDsIfIntersectedRectangle);
+			}
 		}
 		else if (SelectionMode == ESelectionMode::Triangle)
 		{
 			TRACE_CPUPROFILER_EVENT_SCOPE(MeshSelectionMechanic_OnDragRectangleChanged_Triangle);
 
-			RectangleSelectedIDs = FindAllIntersectionsAxisAlignedBox2(
-				Tree, RectangleXY, AppendTriangleID, AppendTriangleIDIfIntersected);
+			if (bSourceIsLivePreview)
+			{
+				RectangleSelectedIDs = FindAllIntersections(
+					Tree, LocalFrustum, AppendTriangleID, AppendTriangleIDIfIntersectedFrustum);
+			}
+			else
+			{
+				RectangleSelectedIDs = FindAllIntersections(
+					Tree, RectangleXY, AppendTriangleID, AppendTriangleIDIfIntersectedRectangle);
+			}
 		}
 		else if (SelectionMode == ESelectionMode::Island)
 		{
 			TRACE_CPUPROFILER_EVENT_SCOPE(MeshSelectionMechanic_OnDragRectangleChanged_Component);
 
-			TArray<int32> SeedTriangles = FindAllIntersectionsAxisAlignedBox2(
-				Tree, RectangleXY, AppendTriangleID, AppendTriangleIDIfIntersected);
+			TArray<int32> SeedTriangles;
+			if (bSourceIsLivePreview)
+			{
+				SeedTriangles = FindAllIntersections(
+					Tree, LocalFrustum, AppendTriangleID, AppendTriangleIDIfIntersectedFrustum);
+			}
+			else
+			{
+				SeedTriangles = FindAllIntersections(
+					Tree, RectangleXY, AppendTriangleID, AppendTriangleIDIfIntersectedRectangle);
+			}
 
 			// TODO(Performance) For large meshes and selections following code is MUCH slower than AABB traversal,
 			//  consider precomputing the connected components an only updating them when the mesh topology changes
 			//  rather than every time the selection changes.
 			FMeshConnectedComponents MeshSelectedComponent(Tree.GetMesh());
-			MeshSelectedComponent.FindTrianglesConnectedToSeeds(SeedTriangles);
+			auto ConnectedComponentsPredicate = [&UnwrapMesh](int32 TriangleA, int32 TriangleB)
+			{
+				if (UnwrapMesh.IsTriangle(TriangleA))
+				{
+					return UnwrapMesh.GetTriNeighbourTris(TriangleA).Contains(TriangleB);
+				}
+				else
+				{
+					return !UnwrapMesh.IsTriangle(TriangleB);
+				}
+			};
+			if (bSourceIsLivePreview)
+			{
+				MeshSelectedComponent.FindTrianglesConnectedToSeeds(SeedTriangles, ConnectedComponentsPredicate);
+			}
+			else
+			{
+				MeshSelectedComponent.FindTrianglesConnectedToSeeds(SeedTriangles);
+			}
 			for (int ComponentIndex = 0; ComponentIndex < MeshSelectedComponent.Components.Num(); ComponentIndex++)
 			{
 				RectangleSelectedIDs.Append(MoveTemp(MeshSelectedComponent.Components[ComponentIndex].Indices));
@@ -1045,8 +1535,17 @@ void UUVEditorMeshSelectionMechanic::OnDragRectangleChanged(const FCameraRectang
 
 			// TODO: This shouldn't be a "find all". We can return early after the first success
 			// since we're selecting the whole mesh
-			TArray<int32> SelectedIDs = FindAllIntersectionsAxisAlignedBox2(
-				Tree, RectangleXY, AppendTriangleID, AppendTriangleIDIfIntersected);
+			TArray<int32> SelectedIDs;
+			if (bSourceIsLivePreview)
+			{
+				SelectedIDs = FindAllIntersections(
+					Tree, LocalFrustum, AppendTriangleID, AppendTriangleIDIfIntersectedFrustum);
+			}
+			else
+			{
+				SelectedIDs = FindAllIntersections(
+					Tree, RectangleXY, AppendTriangleID, AppendTriangleIDIfIntersectedRectangle);
+			}
 			if (!SelectedIDs.IsEmpty())
 			{
 				for (int32 Tid : Tree.GetMesh()->TriangleIndicesItr())
@@ -1060,56 +1559,69 @@ void UUVEditorMeshSelectionMechanic::OnDragRectangleChanged(const FCameraRectang
 			checkSlow(false);
 		}
 
+		if (bSourceIsLivePreview)
+		{
+			RectangleSelectedIDs = Convert3DHitsTo2DHits(SelectionMode, UnwrapMesh, Mesh, UVOverlay, RectangleSelectedIDs, UnsetRectangleSelectedIDs);
+		}
+
 		// See if we have an object in our selection list that corresponds to this asset
 		const FUVToolSelection* PreDragSelection = AssetIDToPreDragSelection[AssetID];
+		const FUVToolSelection* PreDragUnsetSelection = AssetIDToPreDragUnsetSelection[AssetID];
 
-		if (RectangleSelectedIDs.IsEmpty())
+		auto ProcessSelection = [this, &AssetID, &SelectionType](const TArray<int32>& IdsIn, const FUVToolSelection* OldSelection, TArray<FUVToolSelection>& Selections) 
 		{
-			if (!ShouldRestartSelection() && PreDragSelection)
+			if (IdsIn.IsEmpty())
 			{
-				// Keep the existing selection object with no modification.
-				NewSelections.Emplace(*PreDragSelection);
+				if (!ShouldRestartSelection() && OldSelection)
+				{
+					// Keep the existing selection object with no modification.
+					Selections.Emplace(*OldSelection);
+				}
 			}
-		}
-		else if (ShouldRestartSelection() || (!PreDragSelection && !ShouldRemoveFromSelection()))
-		{
-			// Make a new selection object
-			NewSelections.Emplace();
-			NewSelections.Last().Target = Targets[AssetID];
-			NewSelections.Last().Type = SelectionType;
-			NewSelections.Last().SelectedIDs.Append(RectangleSelectedIDs);
-		}
-		else if (PreDragSelection)
-		{
-			// Modify the existing selection object
-			FUVToolSelection NewSelection(*PreDragSelection);
-			ModifyExistingSelection(NewSelection.SelectedIDs, RectangleSelectedIDs);
+			else if (ShouldRestartSelection() || (!OldSelection && !ShouldRemoveFromSelection()))
+			{
+				// Make a new selection object
+				Selections.Emplace();
+				Selections.Last().Target = Targets[AssetID];
+				Selections.Last().Type = SelectionType;
+				Selections.Last().SelectedIDs.Append(IdsIn);
+			}
+			else if (OldSelection)
+			{
+				// Modify the existing selection object
+				FUVToolSelection NewSelection(*OldSelection);
+				ModifyExistingSelection(NewSelection.SelectedIDs, IdsIn);
 
-			// The object may become empty from a removal or toggle, in which case don't add it.
-			if (!NewSelection.IsEmpty())
+				// The object may become empty from a removal or toggle, in which case don't add it.
+				if (!NewSelection.IsEmpty())
+				{
+					Selections.Add(MoveTemp(NewSelection));
+				}
+			}
+			else
 			{
-				NewSelections.Add(MoveTemp(NewSelection));
+				// The only way we can get here is if didn't have an existing selection and were trying
+				// to remove selection, in which case we do nothing.
+				ensure(!OldSelection && ShouldRemoveFromSelection());
 			}
-		}
-		else
-		{
-			// The only way we can get here is if didn't have an existing selection and were trying
-			// to remove selection, in which case we do nothing.
-			ensure(!PreDragSelection && ShouldRemoveFromSelection());
-		}
+		};
+		ProcessSelection(RectangleSelectedIDs, PreDragSelection, NewSelections);
+		ProcessSelection(UnsetRectangleSelectedIDs, PreDragUnsetSelection, NewUnsetSelections);
 	}
 
 	SelectionAPI->SetSelections(NewSelections, false, false);
+	SelectionAPI->SetUnsetElementAppliedMeshSelections(NewUnsetSelections, false, false);
 	OnDragSelectionChanged.Broadcast();
 }
 
-void UUVEditorMeshSelectionMechanic::OnDragRectangleFinished(const FCameraRectangle& CurrentRectangle, bool bCancelled)
+
+void UUVEditorMeshSelectionMechanic::OnDragRectangleFinished(const FCameraRectangle& CurrentRectangle, bool bCancelled, bool bSourceIsLivePreview)
 {
 	TRACE_CPUPROFILER_EVENT_SCOPE(MeshSelectionMechanic_OnDragRectangleFinished); // Mark end of drag sequence
 
 	// TODO(Performance) :DynamicMarqueeSelection Remove this call when marquee selection is fast enough to update
 	//  dynamically for large meshes
-	OnDragRectangleChanged(CurrentRectangle);
+	OnDragRectangleChanged(CurrentRectangle, bSourceIsLivePreview);
 
 	if (!bCancelled)
 	{
@@ -1132,7 +1644,7 @@ void UUVEditorMeshSelectionMechanic::OnUpdateModifierState(int ModifierID, bool
 	}
 }
 
-FInputRayHit UUVEditorMeshSelectionMechanic::BeginHoverSequenceHitTest(const FInputDeviceRay& PressPos)
+FInputRayHit UUVEditorMeshSelectionMechanic::BeginHoverSequenceHitTest(const FInputDeviceRay& PressPos, bool bSourceIsLivePreview)
 {
 	FInputRayHit Hit;
 	if (!bIsEnabled || !bShowHoveredElements || SelectionMode == ESelectionMode::None)
@@ -1150,7 +1662,7 @@ FInputRayHit UUVEditorMeshSelectionMechanic::BeginHoverSequenceHitTest(const FIn
 	// We don't bother with the depth since everything is in the same plane.
 	int32 Tid = IndexConstants::InvalidID;
 	int32 AssetID = IndexConstants::InvalidID;
-	Hit.bHit = GetHitTid(PressPos, Tid, AssetID);
+	Hit.bHit = GetHitTid(PressPos, Tid, AssetID, !bSourceIsLivePreview);
 
 	return Hit;
 }
@@ -1159,7 +1671,7 @@ void UUVEditorMeshSelectionMechanic::OnBeginHover(const FInputDeviceRay& DeviceP
 {
 }
 
-bool UUVEditorMeshSelectionMechanic::OnUpdateHover(const FInputDeviceRay& DevicePos)
+bool UUVEditorMeshSelectionMechanic::OnUpdateHover(const FInputDeviceRay& DevicePos, bool bSourceIsLivePreview)
 {
 	using namespace UVEditorMeshSelectionMechanicLocals;
 
@@ -1177,63 +1689,110 @@ bool UUVEditorMeshSelectionMechanic::OnUpdateHover(const FInputDeviceRay& Device
 	HoverPointSet->Clear();
 	HoverLineSet->Clear();
 	HoverTriangleSet->Clear();
+	LivePreviewHoverPointSet->Clear();
+	LivePreviewHoverLineSet->Clear();
+	LivePreviewHoverTriangleSet->Clear();
 
 	int32 Tid = IndexConstants::InvalidID;
 	int32 AssetID = IndexConstants::InvalidID;
-	if (!GetHitTid(DevicePos, Tid, AssetID))
+	if (!GetHitTid(DevicePos, Tid, AssetID, !bSourceIsLivePreview))
 	{
 		return false;
 	}
+	FDynamicMesh3* UnwrapMesh = Targets[AssetID]->UnwrapCanonical.Get();
+	FDynamicMesh3* AppliedMesh = Targets[AssetID]->AppliedCanonical.Get();
+	FDynamicMesh3* Mesh = bSourceIsLivePreview ? AppliedMesh : UnwrapMesh;
+	FDynamicMeshUVOverlay& UVOverlay = *AppliedMesh->Attributes()->GetUVLayer(Targets[AssetID]->UVLayerIndex);
 
-	FDynamicMesh3* Mesh = Targets[AssetID]->UnwrapCanonical.Get();
-
-	TArray<int32> ConvertedIDs;
+	TArray<int32> Converted3DIDs, Unset3DIDs;
+	TArray<int32> Converted2DIDs;
 	if (SelectionMode == ESelectionMode::Vertex || SelectionMode == ESelectionMode::Edge)
 	{
-		ConvertToHitElementList(SelectionMode, *Mesh,
-			Tid, CameraState, DevicePos.WorldRay, ConvertedIDs);
-		if (ConvertedIDs.IsEmpty())
+		if (bSourceIsLivePreview)
 		{
-			// We were too far from a vert or edge, probably.
-			return false;
+			ConvertToHitElementList(SelectionMode, *Mesh, *UnwrapMesh, UVOverlay,
+				Tid, CameraState, DevicePos.WorldRay, Converted3DIDs);
+			Converted2DIDs = Convert3DHitsTo2DHits(SelectionMode, *UnwrapMesh, *AppliedMesh, UVOverlay, Converted3DIDs, Unset3DIDs);
+		}
+		else
+		{
+			ConvertToHitElementList(SelectionMode, *Mesh, *UnwrapMesh, UVOverlay,
+				Tid, CameraState, DevicePos.WorldRay, Converted2DIDs);
+			Converted3DIDs = Convert2DHitsTo3DHits(SelectionMode, *UnwrapMesh, *AppliedMesh, UVOverlay, Converted2DIDs);
 		}
 	}
 
-	if (SelectionMode == ESelectionMode::Vertex)
+
+	auto SetupVertexHighlight = [](TArray<int32>& IDs, FDynamicMesh3& Mesh, UPointSetComponent& PointSet)
 	{
-		const FVector3d& P = Mesh->GetVertexRef(ConvertedIDs[0]);
-		const FRenderablePoint PointToRender(P,
-			FUVEditorUXSettings::SelectionHoverTriangleWireframeColor,
-			FUVEditorUXSettings::SelectionPointThickness);
+		for (int32 ID : IDs)
+		{
+			const FVector3d& P = Mesh.GetVertexRef(ID);
+			const FRenderablePoint PointToRender(P,
+				FUVEditorUXSettings::SelectionHoverTriangleWireframeColor,
+				FUVEditorUXSettings::SelectionPointThickness);
 
-		HoverPointSet->AddPoint(PointToRender);
-	}
-	else if (SelectionMode == ESelectionMode::Edge)
+			PointSet.AddPoint(PointToRender);
+		}
+	};
+
+	auto SetupEdgeHighlight = [](TArray<int32>& IDs, FDynamicMesh3& Mesh, ULineSetComponent& LineSet)
 	{
-		const FIndex2i EdgeVids = Mesh->GetEdgeV(ConvertedIDs[0]);
-		const FVector& A = Mesh->GetVertexRef(EdgeVids.A);
-		const FVector& B = Mesh->GetVertexRef(EdgeVids.B);
+		for (int32 ID : IDs)
+		{
+			const FIndex2i EdgeVids = Mesh.GetEdgeV(ID);
+			const FVector& A = Mesh.GetVertexRef(EdgeVids.A);
+			const FVector& B = Mesh.GetVertexRef(EdgeVids.B);
+
+			LineSet.AddLine(A, B,
+				FUVEditorUXSettings::SelectionHoverTriangleWireframeColor,
+				FUVEditorUXSettings::SelectionLineThickness,
+				FUVEditorUXSettings::SelectionHoverWireframeDepthBias);
+		}
+	};
 
-		HoverLineSet->AddLine(A, B,
-			FUVEditorUXSettings::SelectionHoverTriangleWireframeColor,
-			FUVEditorUXSettings::SelectionLineThickness,
-			FUVEditorUXSettings::SelectionHoverWireframeDepthBias);
-	}
-	else
+	auto SetupTriangleHighlight = [this, &UnwrapMesh, &Tid](FDynamicMesh3& Mesh, ULineSetComponent& LineSet, UTriangleSetComponent& TriangleSet)
 	{
-		const FIndex3i Vids = Mesh->GetTriangle(Tid);
-		const FVector& A = Mesh->GetVertex(Vids[0]);
-		const FVector& B = Mesh->GetVertex(Vids[1]);
-		const FVector& C = Mesh->GetVertex(Vids[2]);
+		if (!UnwrapMesh->IsTriangle(Tid))
+		{
+			return;
+		}
+
+		const FIndex3i Vids = Mesh.GetTriangle(Tid);
+		const FVector& A = Mesh.GetVertex(Vids[0]);
+		const FVector& B = Mesh.GetVertex(Vids[1]);
+		const FVector& C = Mesh.GetVertex(Vids[2]);
 
-		HoverLineSet->AddLine(A, B, FUVEditorUXSettings::SelectionHoverTriangleWireframeColor,
+		LineSet.AddLine(A, B, FUVEditorUXSettings::SelectionHoverTriangleWireframeColor,
 			FUVEditorUXSettings::SelectionLineThickness, FUVEditorUXSettings::SelectionHoverWireframeDepthBias);
-		HoverLineSet->AddLine(B, C, FUVEditorUXSettings::SelectionHoverTriangleWireframeColor,
+		LineSet.AddLine(B, C, FUVEditorUXSettings::SelectionHoverTriangleWireframeColor,
 			FUVEditorUXSettings::SelectionLineThickness, FUVEditorUXSettings::SelectionHoverWireframeDepthBias);
-		HoverLineSet->AddLine(C, A, FUVEditorUXSettings::SelectionHoverTriangleWireframeColor,
+		LineSet.AddLine(C, A, FUVEditorUXSettings::SelectionHoverTriangleWireframeColor,
 			FUVEditorUXSettings::SelectionLineThickness, FUVEditorUXSettings::SelectionHoverWireframeDepthBias);
-		HoverTriangleSet->AddTriangle(A, B, C, FVector::ZAxisVector,
+		TriangleSet.AddTriangle(A, B, C, FVector::ZAxisVector,
 			FUVEditorUXSettings::SelectionHoverTriangleFillColor, HoverTriangleSetMaterial);
+	};
+
+	// Generate hover visualizations for moused over geometry. This potentially generates more than one "item"
+	// per mesh, since the 3D item might coorespond to more than one 2D item. In particular, vertices and edges
+	// may have more than one match on the 2D side. We also handle the potential existence of hovering over geometry
+	// that is part of unset UVs, which only generate hover visuals on the 3D side.
+	if (SelectionMode == ESelectionMode::Vertex)
+	{
+		SetupVertexHighlight(Converted2DIDs, *UnwrapMesh, *HoverPointSet);
+		SetupVertexHighlight(Converted3DIDs, *AppliedMesh, *LivePreviewHoverPointSet);
+		SetupVertexHighlight(Unset3DIDs, *AppliedMesh, *LivePreviewHoverPointSet);
+	}
+	else if (SelectionMode == ESelectionMode::Edge)
+	{
+		SetupEdgeHighlight(Converted2DIDs, *UnwrapMesh, *HoverLineSet);
+		SetupEdgeHighlight(Converted3DIDs, *AppliedMesh, *LivePreviewHoverLineSet);
+		SetupEdgeHighlight(Unset3DIDs, *AppliedMesh, *LivePreviewHoverLineSet);
+	}
+	else
+	{
+		SetupTriangleHighlight(*UnwrapMesh, *HoverLineSet, *HoverTriangleSet);
+		SetupTriangleHighlight(*AppliedMesh, *LivePreviewHoverLineSet, *LivePreviewHoverTriangleSet);
 	}
 
 	return true;
@@ -1253,6 +1812,18 @@ void UUVEditorMeshSelectionMechanic::OnEndHover()
 	{
 		HoverTriangleSet->Clear();
 	}
+	if (ensure(LivePreviewHoverPointSet.IsValid()))
+	{
+		LivePreviewHoverPointSet->Clear();
+	}
+	if (ensure(LivePreviewHoverLineSet.IsValid()))
+	{
+		LivePreviewHoverLineSet->Clear();
+	}
+	if (ensure(LivePreviewHoverTriangleSet.IsValid()))
+	{
+		LivePreviewHoverTriangleSet->Clear();
+	}
 }
 
 #undef LOCTEXT_NAMESPACE
diff --git a/Engine/Plugins/Experimental/UVEditor/Source/UVEditorTools/Private/Selection/UVToolSelectionAPI.cpp b/Engine/Plugins/Experimental/UVEditor/Source/UVEditorTools/Private/Selection/UVToolSelectionAPI.cpp
--- a/Engine/Plugins/Experimental/UVEditor/Source/UVEditorTools/Private/Selection/UVToolSelectionAPI.cpp
+++ b/Engine/Plugins/Experimental/UVEditor/Source/UVEditorTools/Private/Selection/UVToolSelectionAPI.cpp
@@ -56,7 +56,8 @@ void UUVToolSelectionAPI::Initialize(
 
 	SelectionMechanic = NewObject<UUVEditorMeshSelectionMechanic>();
 	SelectionMechanic->Setup(MechanicAdapter);
-	SelectionMechanic->Initialize(UnwrapWorld, this);
+	SelectionMechanic->Initialize(UnwrapWorld, LivePreviewAPI->GetLivePreviewWorld(), this);
+
 	UnwrapInputRouter->RegisterSource(MechanicAdapter);
 }
 
@@ -140,8 +141,24 @@ void UUVToolSelectionAPI::SetSelections(const TArray<FUVToolSelection>& Selectio
 {
 	using namespace UVToolSelectionAPILocals;
 
+	bool bUsingExistingTransaction = bEmitChange;
+	if (bEmitChange && ensureMsgf(!PendingSelectionChange && !PendingUnsetSelectionChange,
+		TEXT("SetSelections called with bEmitChange while an existing SelectionAPI change transaction is already open. This should be safe but bEmitChange does not need to be explicitly set.")))
+	{
+		BeginChange();
+		bUsingExistingTransaction = false;
+	}
+	bool bWillCallEndChange = bEmitChange && !bUsingExistingTransaction;
+
+
 	bCachedUnwrapSelectionCentroidValid = false;
 
+	// If we don't match with the current unset selection type, just clear them to keep everything consistent
+	if (CurrentUnsetSelections.Num() > 0 && SelectionsIn.Num() > 0 && SelectionsIn[0].Type != CurrentUnsetSelections[0].Type)
+	{
+		CurrentUnsetSelections.Empty();
+	}
+
 	TArray<FUVToolSelection> NewSelections;
 	for (const FUVToolSelection& NewSelection : SelectionsIn)
 	{
@@ -171,59 +188,35 @@ void UUVToolSelectionAPI::SetSelections(const TArray<FUVToolSelection>& Selectio
 		}
 	}
 
-	if (!DoSelectionSetsDiffer(CurrentSelections, NewSelections))
+	bool bSelectionsDiffer = false;
+	if (!bWillCallEndChange)
 	{
-		return;
-	}
-
-	if (bEmitChange)
-	{
-		EmitChangeAPI->BeginUndoTransaction(SelectionChangeTransactionName);
-	}
-
-	if (bBroadcast)
-	{
-		OnPreSelectionChange.Broadcast(bEmitChange);
-	}
-
-	TUniquePtr<FSelectionChange> SelectionChange;
-	if (bEmitChange)
-	{
-		SelectionChange = MakeUnique<FSelectionChange>();
-		SelectionChange->SetBefore(MoveTemp(CurrentSelections));
+		bSelectionsDiffer = DoSelectionSetsDiffer(CurrentSelections, NewSelections);
 	}
 
 	CurrentSelections = MoveTemp(NewSelections);
 
-	if (bEmitChange)
+	if (bWillCallEndChange)
 	{
-		SelectionChange->SetAfter(CurrentSelections);
-		EmitChangeAPI->EmitToolIndependentChange(this, MoveTemp(SelectionChange),
-			SelectionChangeTransactionName);
+		bSelectionsDiffer = EndChangeAndEmitIfModified(bBroadcast);
 	}
 
-	if (HighlightOptions.bAutoUpdateUnwrap)
+	if (bSelectionsDiffer)
 	{
-		FTransform Transform = FTransform::Identity;
-		if (HighlightOptions.bUseCentroidForUnwrapAutoUpdate)
+
+		if (HighlightOptions.bAutoUpdateUnwrap)
 		{
-			Transform = FTransform(GetUnwrapSelectionCentroid());
+			FTransform Transform = FTransform::Identity;
+			if (HighlightOptions.bUseCentroidForUnwrapAutoUpdate)
+			{
+				Transform = FTransform(GetUnwrapSelectionCentroid());
+			}
+			RebuildUnwrapHighlight(Transform);
+		}
+		if (HighlightOptions.bAutoUpdateApplied)
+		{
+			RebuildAppliedPreviewHighlight();
 		}
-		RebuildUnwrapHighlight(Transform);
-	}
-	if (HighlightOptions.bAutoUpdateApplied)
-	{
-		RebuildAppliedPreviewHighlight();
-	}
-
-	if (bBroadcast)
-	{
-		OnSelectionChanged.Broadcast(bEmitChange);
-	}
-
-	if (bEmitChange)
-	{
-		EmitChangeAPI->EndUndoTransaction();
 	}
 }
 
@@ -296,6 +289,81 @@ FVector3d UUVToolSelectionAPI::GetUnwrapSelectionCentroid(bool bForceRecalculate
 }
 
 
+bool UUVToolSelectionAPI::HaveUnsetElementAppliedMeshSelections() const
+{
+	return CurrentUnsetSelections.Num() > 0;
+}
+
+void UUVToolSelectionAPI::SetUnsetElementAppliedMeshSelections(const TArray<FUVToolSelection>& UnsetSelectionsIn, bool bBroadcast, bool bEmitChange)
+{
+	using namespace UVToolSelectionAPILocals;
+
+	bool bUsingExistingTransaction = bEmitChange;
+	if (bEmitChange && ensureMsgf(!PendingSelectionChange && !PendingUnsetSelectionChange,
+		TEXT("SetUnsetSelections called with bEmitChange while an existing SelectionAPI change transaction is already open. This should be safe but bEmitChange does not need to be explicitly set.")))
+	{
+		BeginChange();
+		bUsingExistingTransaction = false;
+	}
+	bool bWillCallEndChange = bEmitChange && !bUsingExistingTransaction;
+
+	// If we don't match with the current unset selection type, just clear them to keep everything consistent
+	if (CurrentSelections.Num() > 0 && UnsetSelectionsIn.Num() > 0 && UnsetSelectionsIn[0].Type != CurrentSelections[0].Type)
+	{
+		CurrentSelections.Empty();
+	}
+
+	TArray<FUVToolSelection> NewUnsetSelections;
+	for (const FUVToolSelection& NewUnsetSelection : UnsetSelectionsIn)
+	{
+		if (ensure(NewUnsetSelection.Target.IsValid() && NewUnsetSelection.Target->IsValid())
+			// All of the unset selections should match type
+			&& ensure(NewUnsetSelections.Num() == 0 || NewUnsetSelection.Type == NewUnsetSelections[0].Type)
+			// Unset selection must not be empty
+			&& ensure(!NewUnsetSelection.IsEmpty())
+				// Shouldn't have unset selection objects pointing to same target
+			&& ensure(!NewUnsetSelections.FindByPredicate(
+				[&NewUnsetSelection](const FUVToolSelection& ExistingSelectionElement) {
+					return NewUnsetSelection.Target == ExistingSelectionElement.Target;
+				})))
+		{
+			NewUnsetSelections.Add(NewUnsetSelection);
+
+			if (NewUnsetSelection.Target.IsValid() && ensure(NewUnsetSelection.Target->UnwrapCanonical))
+			{
+				// These are unset, so they better not be in the Unwrap mesh
+				checkSlow(!NewUnsetSelection.AreElementsPresentInMesh(*NewUnsetSelection.Target->UnwrapCanonical));
+			}
+		}
+	}
+
+	bool bSelectionsDiffer = false;
+	if (!bWillCallEndChange)
+	{
+		bSelectionsDiffer = DoSelectionSetsDiffer(CurrentUnsetSelections, NewUnsetSelections);
+	}
+
+	CurrentUnsetSelections = MoveTemp(NewUnsetSelections);
+
+	if (bWillCallEndChange)
+	{
+		bSelectionsDiffer = EndChangeAndEmitIfModified(bBroadcast);
+	}
+
+	if (bSelectionsDiffer)
+	{
+		if (HighlightOptions.bAutoUpdateApplied)
+		{
+			RebuildAppliedPreviewHighlight();
+		}
+	}
+}
+
+void UUVToolSelectionAPI::ClearUnsetElementAppliedMeshSelections(bool bBroadcast, bool bEmitChange)
+{
+	SetUnsetElementAppliedMeshSelections(TArray<FUVToolSelection>(), bBroadcast, bEmitChange);
+}
+
 void UUVToolSelectionAPI::SetHighlightVisible(bool bUnwrapHighlightVisible, bool bAppliedHighlightVisible, bool bRebuild)
 {
 	HighlightMechanic->SetIsVisible(bUnwrapHighlightVisible, bAppliedHighlightVisible);
@@ -394,6 +462,8 @@ void UUVToolSelectionAPI::RebuildAppliedPreviewHighlight()
 		HighlightMechanic->RebuildAppliedHighlightFromUnwrapSelection(CurrentSelections,
 			HighlightOptions.bBaseHighlightOnPreviews);
 	}
+	HighlightMechanic->AppendAppliedHighlight(CurrentUnsetSelections,
+		HighlightOptions.bBaseHighlightOnPreviews);
 }
 
 void UUVToolSelectionAPI::Render(IToolsContextRenderAPI* RenderAPI)
@@ -412,51 +482,85 @@ void UUVToolSelectionAPI::DrawHUD(FCanvas* Canvas, IToolsContextRenderAPI* Rende
 }
 void UUVToolSelectionAPI::LivePreviewRender(IToolsContextRenderAPI* RenderAPI)
 {
-	//TODO: here we'd route the render call to whatever mechanic is dealing with selection in live preview
+	if (SelectionMechanic)
+	{
+		SelectionMechanic->LivePreviewRender(RenderAPI);
+	}
 }
 void UUVToolSelectionAPI::LivePreviewDrawHUD(FCanvas* Canvas, IToolsContextRenderAPI* RenderAPI)
 {
-	//TODO: here we'd route the DrawHUD call to whatever mechanic is dealing with selection in live preview
+	if (SelectionMechanic)
+	{
+		SelectionMechanic->LivePreviewDrawHUD(Canvas, RenderAPI);
+	}
 }
 
 void UUVToolSelectionAPI::BeginChange()
 {
 	PendingSelectionChange = MakeUnique<FSelectionChange>();
+	PendingUnsetSelectionChange = MakeUnique<FUnsetSelectionChange>();
 	FSelectionChange* CastSelectionChange = static_cast<FSelectionChange*>(PendingSelectionChange.Get());
 	CastSelectionChange->SetBefore(GetSelections());
+	FUnsetSelectionChange* CastUnsetSelectionChange = static_cast<FUnsetSelectionChange*>(PendingUnsetSelectionChange.Get());
+	CastUnsetSelectionChange->SetBefore(GetUnsetElementAppliedMeshSelections());
+
 }
 
 bool UUVToolSelectionAPI::EndChangeAndEmitIfModified(bool bBroadcast)
 {
 	using namespace UVToolSelectionAPILocals;
 
-	if (!PendingSelectionChange)
+	if (!PendingSelectionChange && !PendingUnsetSelectionChange)
 	{
 		return false;
 	}
 
 	FSelectionChange* CastSelectionChange = static_cast<FSelectionChange*>(PendingSelectionChange.Get());
-	
+	FUnsetSelectionChange* CastUnsetSelectionChange = static_cast<FUnsetSelectionChange*>(PendingUnsetSelectionChange.Get());
+
+	bool bSelectionDiffer = true;
+	bool bUnsetSelectionDiffer = true;
 	// See if the selection has changed
 	if (!DoSelectionSetsDiffer(CastSelectionChange->GetBefore(), GetSelections()))
 	{
-		PendingSelectionChange.Reset();
+		PendingSelectionChange.Reset();		
+		bSelectionDiffer = false;
+	}
+	if (!DoSelectionSetsDiffer(CastUnsetSelectionChange->GetBefore(), GetUnsetElementAppliedMeshSelections()))
+	{
+		PendingUnsetSelectionChange.Reset();
+		bUnsetSelectionDiffer = false;
+	}
+	if (!bSelectionDiffer && !bUnsetSelectionDiffer)
+	{
 		return false;
 	}
 
 	EmitChangeAPI->BeginUndoTransaction(SelectionChangeTransactionName);
-	if (bBroadcast)
+	if (bBroadcast && (bSelectionDiffer || bUnsetSelectionDiffer))
+	{
+		OnPreSelectionChange.Broadcast(true, (uint32)((bSelectionDiffer ? ESelectionChangeTypeFlag::SelectionChanged : ESelectionChangeTypeFlag::None) |
+			                                          (bUnsetSelectionDiffer ? ESelectionChangeTypeFlag::UnsetSelectionChanged : ESelectionChangeTypeFlag::None)));
+	}
+	if (bSelectionDiffer)
+	{
+		CastSelectionChange->SetAfter(GetSelections());
+		EmitChangeAPI->EmitToolIndependentChange(this, MoveTemp(PendingSelectionChange),
+			SelectionChangeTransactionName);
+		PendingSelectionChange.Reset();
+	}
+	if (bUnsetSelectionDiffer)
 	{
-		OnPreSelectionChange.Broadcast(true);
+		CastUnsetSelectionChange->SetAfter(GetUnsetElementAppliedMeshSelections());
+		EmitChangeAPI->EmitToolIndependentChange(this, MoveTemp(PendingUnsetSelectionChange),
+			SelectionChangeTransactionName);
+		PendingUnsetSelectionChange.Reset();
 	}
-	CastSelectionChange->SetAfter(GetSelections());
-	EmitChangeAPI->EmitToolIndependentChange(this, MoveTemp(PendingSelectionChange),
-		SelectionChangeTransactionName);
-	PendingSelectionChange.Reset();
 
-	if (bBroadcast)
+	if (bBroadcast && (bSelectionDiffer || bUnsetSelectionDiffer))
 	{
-		OnSelectionChanged.Broadcast(true);
+		OnSelectionChanged.Broadcast(true, (uint32)((bSelectionDiffer ? ESelectionChangeTypeFlag::SelectionChanged : ESelectionChangeTypeFlag::None) |
+													(bUnsetSelectionDiffer ? ESelectionChangeTypeFlag::UnsetSelectionChanged : ESelectionChangeTypeFlag::None)));
 	}
 	EmitChangeAPI->EndUndoTransaction();
 
@@ -538,7 +642,6 @@ void UUVToolSelectionAPI::FSelectionChange::Apply(UObject* Object)
 				}
 			}
 		}
-
 		SelectionAPI->SetSelections(After, true, false);
 	}
 }
@@ -560,7 +663,6 @@ void UUVToolSelectionAPI::FSelectionChange::Revert(UObject* Object)
 				}
 			}
 		}
-
 		SelectionAPI->SetSelections(Before, true, false);
 	}
 }
@@ -570,5 +672,42 @@ FString UUVToolSelectionAPI::FSelectionChange::ToString() const
 	return TEXT("UUVToolSelectionAPI::FSelectionChange");
 }
 
+void UUVToolSelectionAPI::FUnsetSelectionChange::SetBefore(TArray<FUVToolSelection> SelectionsIn)
+{
+	Before = MoveTemp(SelectionsIn);
+}
+
+void UUVToolSelectionAPI::FUnsetSelectionChange::SetAfter(TArray<FUVToolSelection> SelectionsIn)
+{
+	After = MoveTemp(SelectionsIn);
+}
+
+const TArray<FUVToolSelection>& UUVToolSelectionAPI::FUnsetSelectionChange::GetBefore() const
+{
+	return Before;
+}
+
+void UUVToolSelectionAPI::FUnsetSelectionChange::Apply(UObject* Object)
+{
+	UUVToolSelectionAPI* SelectionAPI = Cast<UUVToolSelectionAPI>(Object);
+	if (SelectionAPI)
+	{
+		SelectionAPI->SetUnsetElementAppliedMeshSelections(After, true, false);
+	}
+}
+
+void UUVToolSelectionAPI::FUnsetSelectionChange::Revert(UObject* Object)
+{
+	UUVToolSelectionAPI* SelectionAPI = Cast<UUVToolSelectionAPI>(Object);
+	if (SelectionAPI)
+	{
+		SelectionAPI->SetUnsetElementAppliedMeshSelections(Before, true, false);
+	}
+}
+
+FString UUVToolSelectionAPI::FUnsetSelectionChange::ToString() const
+{
+	return TEXT("UUVToolSelectionAPI::FUnsetSelectionChange");
+}
 
 #undef LOCTEXT_NAMESPACE
\ No newline at end of file
diff --git a/Engine/Plugins/Experimental/UVEditor/Source/UVEditorTools/Private/Selection/UVToolSelectionHighlightMechanic.cpp b/Engine/Plugins/Experimental/UVEditor/Source/UVEditorTools/Private/Selection/UVToolSelectionHighlightMechanic.cpp
--- a/Engine/Plugins/Experimental/UVEditor/Source/UVEditorTools/Private/Selection/UVToolSelectionHighlightMechanic.cpp
+++ b/Engine/Plugins/Experimental/UVEditor/Source/UVEditorTools/Private/Selection/UVToolSelectionHighlightMechanic.cpp
@@ -33,12 +33,12 @@ void UUVToolSelectionHighlightMechanic::Initialize(UWorld* UnwrapWorld, UWorld*
 	{
 		LivePreviewGeometryActor->Destroy();
 	}
-	
+
 	// Owns most of the unwrap geometry except for the unselected paired edges, since we don't
 	// want those to move if we change the actor transform via SetUnwrapHighlightTransform
 	UnwrapGeometryActor = UnwrapWorld->SpawnActor<APreviewGeometryActor>(
 		FVector::ZeroVector, FRotator(0, 0, 0), FActorSpawnParameters());
-	
+
 	UnwrapTriangleSet = NewObject<UTriangleSetComponent>(UnwrapGeometryActor);
 	// We are setting the TranslucencySortPriority here to handle the UV editor's use case in 2D
 	// where multiple translucent layers are drawn on top of each other but still need depth sorting.
@@ -85,14 +85,14 @@ void UUVToolSelectionHighlightMechanic::Initialize(UWorld* UnwrapWorld, UWorld*
 	// Owns the highlights in the live preview.
 	LivePreviewGeometryActor = LivePreviewWorld->SpawnActor<APreviewGeometryActor>(
 		FVector::ZeroVector, FRotator(0, 0, 0), FActorSpawnParameters());
-	
+
 	LivePreviewLineSet = NewObject<ULineSetComponent>(LivePreviewGeometryActor);
 	LivePreviewLineSet->SetLineMaterial(ToolSetupUtil::GetDefaultLineComponentMaterial(
 		GetParentTool()->GetToolManager(), /*bDepthTested*/ true));
 	LivePreviewGeometryActor->SetRootComponent(LivePreviewLineSet.Get());
 	LivePreviewLineSet->RegisterComponent();
-	
-	LivePreviewPointSet = NewObject<UPointSetComponent>(LivePreviewGeometryActor);	
+
+	LivePreviewPointSet = NewObject<UPointSetComponent>(LivePreviewGeometryActor);
 	LivePreviewPointSet->SetPointMaterial(ToolSetupUtil::GetDefaultPointComponentMaterial(
 		GetParentTool()->GetToolManager(), /*bDepthTested*/ true));
 	LivePreviewPointSet->AttachToComponent(LivePreviewLineSet.Get(), FAttachmentTransformRules::KeepWorldTransform);
@@ -440,6 +440,99 @@ void UUVToolSelectionHighlightMechanic::RebuildAppliedHighlightFromUnwrapSelecti
 	}//end for selection
 }
 
+void UUVToolSelectionHighlightMechanic::AppendAppliedHighlight(const TArray<FUVToolSelection>& AppliedSelections, bool bUsePreviews)
+{
+	if (!ensure(LivePreviewGeometryActor))
+	{
+		return;
+	}
+
+	for (const FUVToolSelection& Selection : AppliedSelections)
+	{
+		if (!ensure(Selection.Target.IsValid() && Selection.Target->IsValid()))
+		{
+			return;
+		}
+
+		UUVEditorToolMeshInput* Target = Selection.Target.Get();
+
+		const FDynamicMesh3& AppliedMesh = bUsePreviews ? *Target->AppliedPreview->PreviewMesh->GetMesh()
+			: *Target->AppliedCanonical;
+
+		FTransform MeshTransform = Target->AppliedPreview->PreviewMesh->GetTransform();
+
+		auto AppendEdgeLine = [this, &AppliedMesh, &MeshTransform](int32 AppliedEid)
+		{
+			FVector3d Vert1, Vert2;
+			AppliedMesh.GetEdgeV(AppliedEid, Vert1, Vert2);
+
+			LivePreviewLineSet->AddLine(
+				MeshTransform.TransformPosition(Vert1),
+				MeshTransform.TransformPosition(Vert2),
+				FUVEditorUXSettings::SelectionTriangleWireframeColor,
+				FUVEditorUXSettings::LivePreviewHighlightThickness,
+				FUVEditorUXSettings::LivePreviewHighlightDepthOffset);
+		};
+
+		if (Selection.Type == FUVToolSelection::EType::Triangle)
+		{
+			TRACE_CPUPROFILER_EVENT_SCOPE(Triangle);
+
+			for (int32 Tid : Selection.SelectedIDs)
+			{
+				if (!ensure(AppliedMesh.IsTriangle(Tid)))
+				{
+					continue;
+				}
+
+				// Gather the boundary edges for the live preview
+				FIndex3i TriEids = AppliedMesh.GetTriEdges(Tid);
+				for (int i = 0; i < 3; ++i)
+				{
+					FIndex2i EdgeTids = AppliedMesh.GetEdgeT(TriEids[i]);
+					for (int j = 0; j < 2; ++j)
+					{
+						if (EdgeTids[j] != Tid && !Selection.SelectedIDs.Contains(EdgeTids[j]))
+						{
+							AppendEdgeLine(TriEids[i]);
+							break;
+						}
+					}
+				}//end for tri edges
+			}//end for selection tids
+		}
+		else if (Selection.Type == FUVToolSelection::EType::Edge)
+		{
+			TRACE_CPUPROFILER_EVENT_SCOPE(Edge);
+
+			for (int32 Eid : Selection.SelectedIDs)
+			{
+				if (!ensure(AppliedMesh.IsEdge(Eid)))
+				{
+					continue;
+				}
+
+				AppendEdgeLine(Eid);
+			}
+		}
+		else if (Selection.Type == FUVToolSelection::EType::Vertex)
+		{
+			TRACE_CPUPROFILER_EVENT_SCOPE(Vertex);
+
+			for (int32 Vid : Selection.SelectedIDs)
+			{
+				FVector3d Position = AppliedMesh.GetVertex(Vid);
+
+				LivePreviewPointSet->AddPoint(Position,
+					FUVEditorUXSettings::SelectionTriangleWireframeColor,
+					FUVEditorUXSettings::LivePreviewHighlightPointSize,
+					FUVEditorUXSettings::LivePreviewHighlightDepthOffset);
+			}
+		}
+	}//end for selection
+}
+
+
 void UUVToolSelectionHighlightMechanic::SetEnablePairedEdgeHighlights(bool bEnable)
 {
 	bPairedEdgeHighlightsEnabled = bEnable;
diff --git a/Engine/Plugins/Experimental/UVEditor/Source/UVEditorTools/Private/UVSelectTool.cpp b/Engine/Plugins/Experimental/UVEditor/Source/UVEditorTools/Private/UVSelectTool.cpp
--- a/Engine/Plugins/Experimental/UVEditor/Source/UVEditorTools/Private/UVSelectTool.cpp
+++ b/Engine/Plugins/Experimental/UVEditor/Source/UVEditorTools/Private/UVSelectTool.cpp
@@ -203,7 +203,7 @@ void UUVSelectTool::Setup()
 
 	SelectionAPI->OnSelectionChanged.AddUObject(this, &UUVSelectTool::OnSelectionChanged);
 	SelectionAPI->OnPreSelectionChange.AddWeakLambda(this, 
-		[this](bool bEmitChange) {
+		[this](bool bEmitChange, uint32 SelectionChangeType) {
 			if (bEmitChange)
 			{
 				EmitChangeAPI->EmitToolIndependentChange(GetToolManager(),
@@ -271,7 +271,7 @@ void UUVSelectTool::Setup()
 	if (SelectionAPI->HaveSelections())
 	{
 		// This will also update the gizmo
-		OnSelectionChanged(false);
+		OnSelectionChanged(false, (uint32)ESelectionChangeTypeFlag::SelectionChanged);
 		SelectionAPI->RebuildUnwrapHighlight(TransformGizmo->GetGizmoTransform());
 		SelectionAPI->RebuildAppliedPreviewHighlight();
 	}
@@ -348,12 +348,17 @@ void UUVSelectTool::UpdateGizmo()
 		&& SelectionAPI->HaveSelections());
 }
 
-void UUVSelectTool::OnSelectionChanged(bool)
+void UUVSelectTool::OnSelectionChanged(bool, uint32 SelectionChangeType)
 {
 	TRACE_CPUPROFILER_EVENT_SCOPE(UVSelectTool_OnSelectionChanged);
 	
 	using namespace UVSelectToolLocals;
 
+	if (!(SelectionChangeType & (uint32)ESelectionChangeTypeFlag::SelectionChanged))
+	{
+		return;
+	}
+
 	MovingVidsPerSelection.Reset();
 	MovingVertOriginalPositionsPerSelection.Reset();
 	RenderUpdateTidsPerSelection.Reset();
diff --git a/Engine/Plugins/Experimental/UVEditor/Source/UVEditorTools/Public/ContextObjects/UVToolViewportButtonsAPI.h b/Engine/Plugins/Experimental/UVEditor/Source/UVEditorTools/Public/ContextObjects/UVToolViewportButtonsAPI.h
--- a/Engine/Plugins/Experimental/UVEditor/Source/UVEditorTools/Public/ContextObjects/UVToolViewportButtonsAPI.h
+++ b/Engine/Plugins/Experimental/UVEditor/Source/UVEditorTools/Public/ContextObjects/UVToolViewportButtonsAPI.h
@@ -30,7 +30,7 @@ public:
 		bGizmoButtonsEnabled = bOn;
 	}
 
-	bool AreGizmoButtonsEnabled()
+	bool AreGizmoButtonsEnabled() const
 	{
 		return bGizmoButtonsEnabled;
 	}
@@ -44,7 +44,7 @@ public:
 		}
 	}
 
-	EGizmoMode GetGizmoMode()
+	EGizmoMode GetGizmoMode() const
 	{
 		return GizmoMode;
 	}
@@ -55,7 +55,7 @@ public:
 		bSelectionButtonsEnabled = bOn;
 	}
 
-	bool AreSelectionButtonsEnabled()
+	bool AreSelectionButtonsEnabled() const
 	{
 		return bSelectionButtonsEnabled;
 	}
@@ -69,7 +69,7 @@ public:
 		}
 	}
 
-	ESelectionMode GetSelectionMode()
+	ESelectionMode GetSelectionMode() const
 	{
 		return SelectionMode;
 	}
diff --git a/Engine/Plugins/Experimental/UVEditor/Source/UVEditorTools/Public/Selection/UVEditorMeshSelectionMechanic.h b/Engine/Plugins/Experimental/UVEditor/Source/UVEditorTools/Public/Selection/UVEditorMeshSelectionMechanic.h
--- a/Engine/Plugins/Experimental/UVEditor/Source/UVEditorTools/Public/Selection/UVEditorMeshSelectionMechanic.h
+++ b/Engine/Plugins/Experimental/UVEditor/Source/UVEditorTools/Public/Selection/UVEditorMeshSelectionMechanic.h
@@ -23,16 +23,15 @@ class UMaterialInstanceDynamic;
 class UPointSetComponent;
 class UTriangleSetComponent;
 class UUVToolViewportButtonsAPI;
-
+class ULocalSingleClickInputBehavior;
+class ULocalMouseHoverBehavior;
 
 /**
  * Mechanic for selecting elements of a dynamic mesh in the UV editor. Interacts
  * heavily with UUVToolSelectionAPI, which actually stores selections.
  */
 UCLASS()
-class UVEDITORTOOLS_API UUVEditorMeshSelectionMechanic : public UInteractionMechanic, 
-	public IClickBehaviorTarget,
-	public IHoverBehaviorTarget
+class UVEDITORTOOLS_API UUVEditorMeshSelectionMechanic : public UInteractionMechanic
 {
 	GENERATED_BODY()
 
@@ -49,7 +48,7 @@ public:
 	// The selection API is provided as a parameter rather than being grabbed out of the context
 	// store mainly because UVToolSelectionAPI itself sets up a selection mechanic, and is not 
 	// yet in the context store when it does this. 
-	void Initialize(UWorld* World, UUVToolSelectionAPI* SelectionAPI);
+	void Initialize(UWorld* World, UWorld* LivePreviewWorld, UUVToolSelectionAPI* SelectionAPI);
 	void SetTargets(const TArray<TObjectPtr<UUVEditorToolMeshInput>>& TargetsIn);
 
 	void SetIsEnabled(bool bIsEnabled);
@@ -64,22 +63,23 @@ public:
 	 */
 	void SetSelectionMode(ESelectionMode TargetMode,
 		const FModeChangeOptions& Options = FModeChangeOptions());
-
+	
 	virtual void Render(IToolsContextRenderAPI* RenderAPI) override;
 	virtual void DrawHUD(FCanvas* Canvas, IToolsContextRenderAPI* RenderAPI);
-
+	virtual void LivePreviewRender(IToolsContextRenderAPI* RenderAPI);
+	virtual void LivePreviewDrawHUD(FCanvas* Canvas, IToolsContextRenderAPI* RenderAPI);
 	// IClickBehaviorTarget implementation
-	virtual FInputRayHit IsHitByClick(const FInputDeviceRay& ClickPos) override;
-	virtual void OnClicked(const FInputDeviceRay& ClickPos) override;
+	virtual FInputRayHit IsHitByClick(const FInputDeviceRay& ClickPos, bool bSourceIsLivePreview);
+	virtual void OnClicked(const FInputDeviceRay& ClickPos, bool bSourceIsLivePreview);
 
 	// IModifierToggleBehaviorTarget implementation
-	virtual void OnUpdateModifierState(int ModifierID, bool bIsOn) override;
+	virtual void OnUpdateModifierState(int ModifierID, bool bIsOn);
 
 	// IHoverBehaviorTarget implementation
-	virtual FInputRayHit BeginHoverSequenceHitTest(const FInputDeviceRay& PressPos) override;
-	virtual void OnBeginHover(const FInputDeviceRay& DevicePos) override;
-	virtual bool OnUpdateHover(const FInputDeviceRay& DevicePos) override;
-	virtual void OnEndHover() override;
+	virtual FInputRayHit BeginHoverSequenceHitTest(const FInputDeviceRay& PressPos, bool bSourceIsLivePreview);
+	virtual void OnBeginHover(const FInputDeviceRay& DevicePos);
+	virtual bool OnUpdateHover(const FInputDeviceRay& DevicePos, bool bSourceIsLivePreview);
+	virtual void OnEndHover();
 
 	/**
 	 * Broadcasted whenever the marquee mechanic rectangle is changed, since these changes
@@ -98,8 +98,26 @@ protected:
 	UPROPERTY()
 	TObjectPtr<UUVToolEmitChangeAPI> EmitChangeAPI = nullptr;
 
+	UPROPERTY()
+	TObjectPtr<UUVToolLivePreviewAPI> LivePreviewAPI = nullptr;
+
+	UPROPERTY()
+	TObjectPtr<ULocalSingleClickInputBehavior> UnwrapClickTargetRouter = nullptr;
+
+	UPROPERTY()
+	TObjectPtr<ULocalSingleClickInputBehavior> LivePreviewClickTargetRouter = nullptr;
+
+	UPROPERTY()
+	TObjectPtr<ULocalMouseHoverBehavior> UnwrapHoverBehaviorTargetRouter = nullptr;
+
+	UPROPERTY()
+	TObjectPtr<ULocalMouseHoverBehavior> LivePreviewHoverBehaviorTargetRouter = nullptr;
+
 	UPROPERTY()
 	TObjectPtr<URectangleMarqueeMechanic> MarqueeMechanic = nullptr;
+	
+	UPROPERTY()
+	TObjectPtr<URectangleMarqueeMechanic> LivePreviewMarqueeMechanic = nullptr;
 
 	UPROPERTY()
 	TObjectPtr<UMaterialInstanceDynamic> HoverTriangleSetMaterial = nullptr;
@@ -111,16 +129,33 @@ protected:
 	TWeakObjectPtr<ULineSetComponent> HoverLineSet = nullptr;
 	TWeakObjectPtr<UPointSetComponent> HoverPointSet = nullptr;
 
+	UPROPERTY()
+	TObjectPtr<UInputBehaviorSet> LivePreviewBehaviorSet = nullptr;
+
+	UPROPERTY()
+	TObjectPtr<ULocalInputBehaviorSource> LivePreviewBehaviorSource = nullptr;
+
+	TWeakObjectPtr<UInputRouter> LivePreviewInputRouter = nullptr;
+
+	UPROPERTY()
+	TObjectPtr<APreviewGeometryActor> LivePreviewHoverGeometryActor = nullptr;
+	// Weak pointers so that they go away when geometry actor is destroyed
+	TWeakObjectPtr<UTriangleSetComponent> LivePreviewHoverTriangleSet = nullptr;
+	TWeakObjectPtr<ULineSetComponent> LivePreviewHoverLineSet = nullptr;
+	TWeakObjectPtr<UPointSetComponent> LivePreviewHoverPointSet = nullptr;
+
 	// Should be the same as the mode-level targets array, indexed by AssetID
+	TSharedPtr<FDynamicMeshAABBTree3> GetMeshSpatial(int32 TargetId, bool bUseUnwrap);
 	TArray<TObjectPtr<UUVEditorToolMeshInput>> Targets;
-	TArray<TSharedPtr<FDynamicMeshAABBTree3>> MeshSpatials; // 1:1 with Targets
+	TArray<TSharedPtr<FDynamicMeshAABBTree3>> UnwrapMeshSpatials; // 1:1 with Targets
+	TArray<TSharedPtr<FDynamicMeshAABBTree3>> AppliedMeshSpatials; // 1:1 with Targets
 
 	ESelectionMode SelectionMode;
 	bool bIsEnabled = false;
 	bool bShowHoveredElements = true;
 
 	bool GetHitTid(const FInputDeviceRay& ClickPos, int32& TidOut,
-		int32& AssetIDOut, int32* ExistingSelectionObjectIndexOut = nullptr);
+		int32& AssetIDOut, bool bUseUnwrap, int32* ExistingSelectionObjectIndexOut = nullptr);
 	void ModifyExistingSelection(TSet<int32>& SelectionSetToModify, const TArray<int32>& SelectedIDs);
 
 	FViewCameraState CameraState;
@@ -138,11 +173,14 @@ protected:
 
 	// For marquee mechanic
 	void OnDragRectangleStarted();
-	void OnDragRectangleChanged(const FCameraRectangle& CurrentRectangle);
-	void OnDragRectangleFinished(const FCameraRectangle& Rectangle, bool bCancelled);
+	void OnDragRectangleChanged(const FCameraRectangle& CurrentRectangle, bool bSourceIsLivePreview);
+	void OnDragRectangleFinished(const FCameraRectangle& Rectangle, bool bCancelled, bool bSourceIsLivePreview);
+
 	TArray<FUVToolSelection> PreDragSelections;
+	TArray<FUVToolSelection> PreDragUnsetSelections;
 	// Maps asset id to a pre drag selection so that it is easy to tell which assets
 	// started with a selection. 1:1 with Targets.
 	TArray<const FUVToolSelection*> AssetIDToPreDragSelection;
+	TArray<const FUVToolSelection*> AssetIDToPreDragUnsetSelection;
 };
 
diff --git a/Engine/Plugins/Experimental/UVEditor/Source/UVEditorTools/Public/Selection/UVToolSelectionAPI.h b/Engine/Plugins/Experimental/UVEditor/Source/UVEditorTools/Public/Selection/UVToolSelectionAPI.h
--- a/Engine/Plugins/Experimental/UVEditor/Source/UVEditorTools/Public/Selection/UVToolSelectionAPI.h
+++ b/Engine/Plugins/Experimental/UVEditor/Source/UVEditorTools/Public/Selection/UVToolSelectionAPI.h
@@ -13,6 +13,18 @@ class UUVEditorMechanicAdapterTool;
 class UUVEditorMeshSelectionMechanic;
 class UUVToolSelectionHighlightMechanic;
 
+/*
+ *  Enum to describe the nature of the selection change during either a OnPreSelectionChange callback or
+ *  OnSelectionChange callback.
+ */
+enum class ESelectionChangeTypeFlag : uint8
+{
+	None = 0,
+	SelectionChanged = 1 << 0, // Set for normal selection change events
+	UnsetSelectionChanged = 1 << 1 // Set for unset selection change events
+};
+ENUM_CLASS_FLAGS(ESelectionChangeTypeFlag);
+
 /**
  * API for dealing with mode-level selection in the UV editor.
  * 
@@ -51,7 +63,17 @@ public:
 	 */
 	FUVToolSelection::EType GetSelectionsType()
 	{
-		return CurrentSelections.Num() > 0 ? CurrentSelections[0].Type : FUVToolSelection::EType::Triangle;
+		if (CurrentSelections.Num() > 0)
+		{
+			return CurrentSelections[0].Type;
+		}
+		else if (CurrentUnsetSelections.Num() > 0)
+		{
+			return CurrentUnsetSelections[0].Type;
+		}
+		else {
+			return FUVToolSelection::EType::Triangle;
+		}		
 	}
 
 	/**
@@ -61,6 +83,8 @@ public:
 	 * 
 	 * @param bBroadcast If true, broadcast OnPreSelectionChange and OnSelectionChanged
 	 * @param bEmitChange If true, emit an undo/redo transaction.
+	 * @param bAllowInconsistentSelectionTypes If true, allow setting selections with different selection types than any currently set UnsetElement selections.
+	          Primarily useful during undo/redo
 	 */
 	void SetSelections(const TArray<FUVToolSelection>& SelectionsIn, bool bBroadcast, bool bEmitChange);
 
@@ -76,17 +100,16 @@ public:
 	 */
 	FVector3d GetUnwrapSelectionCentroid(bool bForceRecalculate = false);
 
-	// TODO: When we have support for selection in 3d live preview viewport, we may need these.
-	// These are function that would be used if a selection made in the applied mesh cannot be
+	// These are functions that are used if a selection made in the applied mesh cannot be
 	// converted to an unwrap mesh selection due to unset UV elements (meaning that the selected
-	// portion does not exist in the unwrap mesh). Such selections would be stored separately 
+	// portion does not exist in the unwrap mesh). Such selections are stored separately 
 	// from all those that are able to be converted (which get stored with the usual unwrap selections).
-	// These "unset element" selections would be ignored by many tools/actions but operated on 
+	// These "unset element" selections are ignored by many tools/actions but operated on 
 	// by some others.
-	//bool HaveUnsetElementAppliedMeshSelections() const;
-	//const TArray<FUVToolSelection>& GetUnsetElementAppliedMeshSelections();
-	//void SetUnsetElementAppliedMeshSelections(const TArray<FUVToolSelection>& SelectionsIn, bool bBroadcast, bool bEmitChange);
-	//void ClearUnsetElementAppliedMeshSelections(bool bBroadcast, bool bEmitChange);
+	bool HaveUnsetElementAppliedMeshSelections() const;
+	const TArray<FUVToolSelection>& GetUnsetElementAppliedMeshSelections() { return CurrentUnsetSelections; }
+	void SetUnsetElementAppliedMeshSelections(const TArray<FUVToolSelection>& SelectionsIn, bool bBroadcast, bool bEmitChange);
+	void ClearUnsetElementAppliedMeshSelections(bool bBroadcast, bool bEmitChange);
 
 	/**
 	 * Mode of operation for the selection mechanic.
@@ -199,8 +222,10 @@ public:
 	 * 
 	 * @param bEmitChangeAllowed If false, the callback must not emit any undo/redo transactions,
 	 *  likely because this is being called from an apply/revert of an existing transaction.
+	 * @param SelectionChangeType Set as some combination of ESelectionChangeCondition to describe
+	 *  the nature of the selection change.
 	 */
-	DECLARE_MULTICAST_DELEGATE_OneParam(FOnPreSelectionChange, bool bEmitChangeAllowed);
+	DECLARE_MULTICAST_DELEGATE_TwoParams(FOnPreSelectionChange, bool bEmitChangeAllowed, uint32 SelectionChangeType );
 	FOnPreSelectionChange OnPreSelectionChange;
 
 	/**
@@ -209,8 +234,10 @@ public:
 	 * 
 	 * @param bEmitChangeAllowed If false, the callback must not emit any undo/redo transactions,
 	 *  likely because this is being called from an apply/revert of an existing transaction.
+	 * @param SelectionChangeType Set as some combination of ESelectionChangeCondition to describe
+	 *  the nature of the selection change.
 	 */
-	DECLARE_MULTICAST_DELEGATE_OneParam(FOnSelectionChanged, bool bEmitChangeAllowed);
+	DECLARE_MULTICAST_DELEGATE_TwoParams(FOnSelectionChanged, bool bEmitChangeAllowed, uint32 SelectionChangeType);
 	FOnSelectionChanged OnSelectionChanged;
 
 	/**
@@ -223,7 +250,7 @@ public:
 	// Undo/redo support:
 
 	/**
-	 * An object that can be as an undo/redo item, usually emitted by the selection API
+	 * An object representing a selection that can be as an undo/redo item, usually emitted by the selection API
 	 * itself. Expects UUVToolSelectionAPI to be the associated UObject.
 	 */
 	class FSelectionChange : public FToolCommandChange
@@ -253,6 +280,27 @@ public:
 		TArray<FUVToolSelection> After;
 	};
 
+	/**
+	 * An object representing an unset UV selection that can be as an undo/redo item, usually emitted by the selection API
+	 * itself. Expects UUVToolSelectionAPI to be the associated UObject.
+	 */
+	class FUnsetSelectionChange : public FToolCommandChange
+	{
+	public:
+		virtual void SetBefore(TArray<FUVToolSelection> Selections);
+		virtual void SetAfter(TArray<FUVToolSelection> Selections);
+
+		virtual const TArray<FUVToolSelection>& GetBefore() const;
+
+		virtual void Apply(UObject* Object) override;
+		virtual void Revert(UObject* Object) override;
+		virtual FString ToString() const override;
+
+	protected:
+		TArray<FUVToolSelection> Before;
+		TArray<FUVToolSelection> After;
+	};
+
 	/**
 	 * Preps a selection change transaction, if the user wants more control on what the previous and
 	 * current selection is.
@@ -284,6 +332,7 @@ public:
 	virtual void OnToolEnded(UInteractiveTool* DeadTool) override;
 protected:
 	TArray<FUVToolSelection> CurrentSelections;
+	TArray<FUVToolSelection> CurrentUnsetSelections;
 
 	UPROPERTY()
 	TArray<TObjectPtr<UUVEditorToolMeshInput>> Targets;
@@ -303,6 +352,8 @@ protected:
 	TObjectPtr<UUVToolEmitChangeAPI> EmitChangeAPI = nullptr;
 
 	TUniquePtr<FSelectionChange> PendingSelectionChange = nullptr;
+	TUniquePtr<FUnsetSelectionChange> PendingUnsetSelectionChange = nullptr;
+
 
 	FHighlightOptions HighlightOptions;
 
@@ -324,10 +375,21 @@ class UUVToolSupportsSelection : public UInterface
  * after tool completion.
  * If a tool does inherit from IUVToolSupportsSelection, then the UV editor will
  * not clear the selection before invocation, allowing the tool to use it. However,
- * the tool it is expected to properly deal with selection itself and avoid an
+ * the tool is expected to properly deal with selection itself and avoid an
  * invalid state (including avoiding incorrect undo/redo selection event ordering.
  */
 class IUVToolSupportsSelection
 {
 	GENERATED_BODY()
+
+public:
+
+	/*
+	 * A tool that returns false here will have the unset element selections 
+	 * (ie selections of triangles that don't have associated UV's) cleared before
+	 * its invocation. Returning true will keep these unset element selections, but
+	 * the tool will be responsible for dealing with them and avoiding any invalid states.
+	 */
+	virtual bool SupportsUnsetElementAppliedMeshSelections() const { return false; }
+
 };
diff --git a/Engine/Plugins/Experimental/UVEditor/Source/UVEditorTools/Public/Selection/UVToolSelectionHighlightMechanic.h b/Engine/Plugins/Experimental/UVEditor/Source/UVEditorTools/Public/Selection/UVToolSelectionHighlightMechanic.h
--- a/Engine/Plugins/Experimental/UVEditor/Source/UVEditorTools/Public/Selection/UVToolSelectionHighlightMechanic.h
+++ b/Engine/Plugins/Experimental/UVEditor/Source/UVEditorTools/Public/Selection/UVToolSelectionHighlightMechanic.h
@@ -74,6 +74,14 @@ public:
 	 */
 	void RebuildAppliedHighlightFromUnwrapSelection(const TArray<FUVToolSelection>& UnwrapSelections, bool bUsePreviews = false);
 
+	/**
+	 * Adds additional highlighting elements based on selection information. This is used to extend selection
+	 * line sets for the unset selections if they are needed. 
+	 *
+	 * @param bUsePreviews If true, the previews are used rather than the canonicals for finding the highlight positions.
+	 */
+	void AppendAppliedHighlight(const TArray<FUVToolSelection>& AppliedSelections, bool bUsePreviews = false);
+
 	/**
 	 * When true, edge selections also highlight the edge pairings (i.e. edges that could weld with the
 	 * selected edges). Note that the call itself controls visibility and the next rebuild call, but 
diff --git a/Engine/Plugins/Experimental/UVEditor/Source/UVEditorTools/Public/UVSelectTool.h b/Engine/Plugins/Experimental/UVEditor/Source/UVEditorTools/Public/UVSelectTool.h
--- a/Engine/Plugins/Experimental/UVEditor/Source/UVEditorTools/Public/UVSelectTool.h
+++ b/Engine/Plugins/Experimental/UVEditor/Source/UVEditorTools/Public/UVSelectTool.h
@@ -68,7 +68,7 @@ public:
 	virtual bool HasAccept() const override { return false; }
 
 protected:
-	virtual void OnSelectionChanged(bool bEmitChangeAllowed);
+	virtual void OnSelectionChanged(bool bEmitChangeAllowed, uint32 SelectionChangeType);
 
 	// Callbacks we'll receive from the gizmo proxy
 	virtual void GizmoTransformChanged(UTransformProxy* Proxy, FTransform Transform);
diff --git a/Engine/Plugins/Runtime/GeometryProcessing/Source/DynamicMesh/Private/Parameterization/DynamicMeshUVEditor.cpp b/Engine/Plugins/Runtime/GeometryProcessing/Source/DynamicMesh/Private/Parameterization/DynamicMeshUVEditor.cpp
--- a/Engine/Plugins/Runtime/GeometryProcessing/Source/DynamicMesh/Private/Parameterization/DynamicMeshUVEditor.cpp
+++ b/Engine/Plugins/Runtime/GeometryProcessing/Source/DynamicMesh/Private/Parameterization/DynamicMeshUVEditor.cpp
@@ -551,6 +551,10 @@ bool FDynamicMeshUVEditor::SetTriangleUVsFromConformal(const TArray<int32>& Tria
 
 	for (int32 tid : Triangles)
 	{
+		if (bUseExistingUVTopology && !UVOverlay->IsSetTriangle(tid))
+		{
+			continue;
+		}
 		FIndex3i Triangle = (bUseExistingUVTopology) ? UVOverlay->GetTriangle(tid) : Mesh->GetTriangle(tid);
 		FIndex3i NewTriangle;
 		for (int32 j = 0; j < 3; ++j)
diff --git a/Engine/Plugins/Runtime/MeshModelingToolset/Source/ModelingComponents/Private/Mechanics/RectangleMarqueeMechanic.cpp b/Engine/Plugins/Runtime/MeshModelingToolset/Source/ModelingComponents/Private/Mechanics/RectangleMarqueeMechanic.cpp
--- a/Engine/Plugins/Runtime/MeshModelingToolset/Source/ModelingComponents/Private/Mechanics/RectangleMarqueeMechanic.cpp
+++ b/Engine/Plugins/Runtime/MeshModelingToolset/Source/ModelingComponents/Private/Mechanics/RectangleMarqueeMechanic.cpp
@@ -224,7 +224,14 @@ TPair<FInputCapturePriority, FInputCapturePriority> URectangleMarqueeMechanic::G
 void URectangleMarqueeMechanic::Render(IToolsContextRenderAPI* RenderAPI)
 {
 	// Cache the camera state
-	GetParentTool()->GetToolManager()->GetContextQueriesAPI()->GetCurrentViewState(CameraRectangle.CameraState);
+	if (bUseExternalUpdateCameraState && ensureMsgf(UpdateCameraStateFunc, TEXT("bUseExternalUpdateCameraState is requested, but no update function has been provided.")))
+	{
+		CameraRectangle.CameraState = UpdateCameraStateFunc();
+	}
+	else
+	{
+		GetParentTool()->GetToolManager()->GetContextQueriesAPI()->GetCurrentViewState(CameraRectangle.CameraState);
+	}
 }
 
 FInputRayHit URectangleMarqueeMechanic::CanBeginClickDragSequence(const FInputDeviceRay& PressPos)
diff --git a/Engine/Plugins/Runtime/MeshModelingToolset/Source/ModelingComponents/Public/Mechanics/RectangleMarqueeMechanic.h b/Engine/Plugins/Runtime/MeshModelingToolset/Source/ModelingComponents/Public/Mechanics/RectangleMarqueeMechanic.h
--- a/Engine/Plugins/Runtime/MeshModelingToolset/Source/ModelingComponents/Public/Mechanics/RectangleMarqueeMechanic.h
+++ b/Engine/Plugins/Runtime/MeshModelingToolset/Source/ModelingComponents/Public/Mechanics/RectangleMarqueeMechanic.h
@@ -134,6 +134,16 @@ public:
 	UPROPERTY()
 	bool bUseExternalClickDragBehavior = false;
 
+	/**
+	 * If true, then the URectangleMarqueeMechanic will not query the tool manager for the current camera view state,
+	 but instead call the custom provided function, UpdateCameraStateFunc, to recieve the correct camera state for the 
+	 CameraRectangle.
+	 */
+	UPROPERTY()
+	bool bUseExternalUpdateCameraState = false;
+
+	TUniqueFunction< FViewCameraState() > UpdateCameraStateFunc;
+
 	/**
 	 * If the computation time for a single call to OnDragRectangleChanged ever exceeds this threshold then future calls
 	 * to this function (in the current drag sequence) will be deferred until the mouse button is released. This will
diff --git a/Engine/Source/Runtime/InteractiveToolsFramework/Public/BaseBehaviors/SingleClickBehavior.h b/Engine/Source/Runtime/InteractiveToolsFramework/Public/BaseBehaviors/SingleClickBehavior.h
--- a/Engine/Source/Runtime/InteractiveToolsFramework/Public/BaseBehaviors/SingleClickBehavior.h
+++ b/Engine/Source/Runtime/InteractiveToolsFramework/Public/BaseBehaviors/SingleClickBehavior.h
@@ -96,6 +96,9 @@ public:
 	/** lambda implementation of OnClicked */
 	TUniqueFunction<void(const FInputDeviceRay&)> OnClickedFunc = [](const FInputDeviceRay& ClickPos) {};
 
+	/** lambda implementation of OnUpdateModifierState */
+	TUniqueFunction< void(int, bool) > OnUpdateModifierStateFunc = [](int ModifierID, bool bIsOn) {};
+
 public:
 	// IClickBehaviorTarget implementation
 
@@ -108,4 +111,10 @@ public:
 	{
 		return OnClickedFunc(ClickPos);
 	}
+
+	// IModifierToggleBehaviorTarget implementation
+	virtual void OnUpdateModifierState(int ModifierID, bool bIsOn)
+	{
+		return OnUpdateModifierStateFunc(ModifierID,bIsOn);
+	}
 };
