UVEditor: Address Non-Unity build issue regarding FConvexVolume headers and UVEditorMeshSelectionMechanic.

#rnx

diff --git a/Engine/Plugins/Experimental/UVEditor/Source/UVEditorTools/Private/Selection/UVEditorMeshSelectionMechanic.cpp b/Engine/Plugins/Experimental/UVEditor/Source/UVEditorTools/Private/Selection/UVEditorMeshSelectionMechanic.cpp
--- a/Engine/Plugins/Experimental/UVEditor/Source/UVEditorTools/Private/Selection/UVEditorMeshSelectionMechanic.cpp
+++ b/Engine/Plugins/Experimental/UVEditor/Source/UVEditorTools/Private/Selection/UVEditorMeshSelectionMechanic.cpp
@@ -15,6 +15,7 @@
 #include "Drawing/PreviewGeometryActor.h"
 #include "Intersection/IntrTriangle2AxisAlignedBox2.h"
 #include "Intersection/IntersectionQueries2.h"
+#include "ConvexVolume.h"
 #include "InteractiveToolManager.h"
 #include "MeshOpPreviewHelpers.h" // UMeshOpPreviewWithBackgroundCompute
 #include "Polyline3.h"
