UVEditor: Add icons for the Transform, Align, and Distribute tools.

#rnx

diff --git a/Engine/Plugins/Experimental/UVEditor/Source/UVEditor/Private/UVEditorStyle.cpp b/Engine/Plugins/Experimental/UVEditor/Source/UVEditor/Private/UVEditorStyle.cpp
--- a/Engine/Plugins/Experimental/UVEditor/Source/UVEditor/Private/UVEditorStyle.cpp
+++ b/Engine/Plugins/Experimental/UVEditor/Source/UVEditor/Private/UVEditorStyle.cpp
@@ -28,6 +28,9 @@ FUVEditorStyle::FUVEditorStyle()
 	Set("UVEditor.BeginChannelEditTool",        new IMAGE_BRUSH_SVG("AttributeEditor", ToolbarIconSize));
 	Set("UVEditor.BeginSeamTool",               new IMAGE_BRUSH_SVG("ModelingUVSeamEdit", ToolbarIconSize));
 	Set("UVEditor.BeginRecomputeUVsTool",       new IMAGE_BRUSH_SVG("GroupUnwrap", ToolbarIconSize));
+	Set("UVEditor.BeginTransformTool",          new IMAGE_BRUSH_SVG("TransformUVs", ToolbarIconSize));
+	Set("UVEditor.BeginAlignTool",              new IMAGE_BRUSH_SVG("AlignLeft", ToolbarIconSize));
+	Set("UVEditor.BeginDistributeTool",         new IMAGE_BRUSH_SVG("DistributeHorizontally", ToolbarIconSize));
 
 	// Select tool actions
 	Set("UVEditor.SewAction", new IMAGE_BRUSH_SVG("UVSew", ToolbarIconSize));
