UVEditor: Add support for a quick translate and quick rotate to the Transform tool. These controls provide single action applications of translations and rotations for rapid editing.

diff --git a/Engine/Plugins/Experimental/UVEditor/Source/UVEditorTools/Private/Operators/UVEditorUVTransformOp.cpp b/Engine/Plugins/Experimental/UVEditor/Source/UVEditorTools/Private/Operators/UVEditorUVTransformOp.cpp
--- a/Engine/Plugins/Experimental/UVEditor/Source/UVEditorTools/Private/Operators/UVEditorUVTransformOp.cpp
+++ b/Engine/Plugins/Experimental/UVEditor/Source/UVEditorTools/Private/Operators/UVEditorUVTransformOp.cpp
@@ -189,12 +189,12 @@ void FUVEditorUVTransformBaseOp::CalculateResult(FProgressCancel* Progress)
 	HandleTransformationOp(Progress);
 }
 
-FVector2f FUVEditorUVTransformOp::GetPivotFromMode(int32 ElementID)
+FVector2f FUVEditorUVTransformOp::GetPivotFromMode(int32 ElementID, EUVEditorPivotTypeBackend Mode)
 {
 	FVector2f Pivot;
 	const int32* Component;
 
-	switch (PivotMode)
+	switch (Mode)
 	{
 	case EUVEditorPivotTypeBackend::Origin:
 		Pivot = FVector2f(0, 0);
@@ -218,9 +218,9 @@ FVector2f FUVEditorUVTransformOp::GetPivotFromMode(int32 ElementID)
 
 void FUVEditorUVTransformOp::HandleTransformationOp(FProgressCancel * Progress)
 {
-	auto ScaleFunc = [&](int32 Vid)
+	auto ScaleFunc = [this](int32 Vid)
 	{
-		FVector2f ScalePivot = GetPivotFromMode(Vid);
+		FVector2f ScalePivot = GetPivotFromMode(Vid, PivotMode);
 		FVector2f UV = FUVEditorUXSettings::InternalUVToExternalUV(ActiveUVLayer->GetElement(Vid));
 		UV = (UV - ScalePivot);
 		UV[0] *= Scale[0];
@@ -229,32 +229,50 @@ void FUVEditorUVTransformOp::HandleTransformationOp(FProgressCancel * Progress)
 		ActiveUVLayer->SetElement(Vid, FUVEditorUXSettings::ExternalUVToInternalUV(UV));
 	};
 
-	auto RotFunc = [&](int32 Vid)
+	auto BaseRotFunc = [this](int32 Vid, float RotationIn, const FVector2f& Pivot)
 	{
-		FVector2f RotationPivot = GetPivotFromMode(Vid);
 		FVector2f UV = FUVEditorUXSettings::InternalUVToExternalUV(ActiveUVLayer->GetElement(Vid));
 		FVector2f UV_Rotated;
-		double RotationInRadians = Rotation / 180.0 * UE_PI;
-		UV = (UV - RotationPivot);
+		double RotationInRadians = RotationIn / 180.0 * UE_PI;
+		UV = (UV - Pivot);
 		UV_Rotated[0] = UV[0] * FMath::Cos(RotationInRadians) - UV[1] * FMath::Sin(RotationInRadians);
 		UV_Rotated[1] = UV[0] * FMath::Sin(RotationInRadians) + UV[1] * FMath::Cos(RotationInRadians);
-		UV = (UV_Rotated + RotationPivot);
+		UV = (UV_Rotated + Pivot);
 		ActiveUVLayer->SetElement(Vid, FUVEditorUXSettings::ExternalUVToInternalUV(UV));
 	};
 
-	auto TranslateFunc = [&](int32 Vid)
+	auto QuickRotFunc = [this, &BaseRotFunc](int32 Vid)
+	{
+		FVector2f RotationPivot = GetPivotFromMode(Vid, EUVEditorPivotTypeBackend::BoundingBoxCenter);
+		BaseRotFunc(Vid, QuickRotation, RotationPivot);
+	};
+
+	auto RotFunc = [this, &BaseRotFunc](int32 Vid)
+	{
+		FVector2f RotationPivot = GetPivotFromMode(Vid, PivotMode);
+		BaseRotFunc(Vid, Rotation, RotationPivot);
+	};
+
+	auto QuickTranslateFunc = [this](int32 Vid)
+	{
+		FVector2f UV = FUVEditorUXSettings::InternalUVToExternalUV(ActiveUVLayer->GetElement(Vid));
+		UV = (UV + FVector2f(QuickTranslation));
+		ActiveUVLayer->SetElement(Vid, FUVEditorUXSettings::ExternalUVToInternalUV(UV));
+	};
+
+	auto TranslateFunc = [this](int32 Vid)
 	{	
 		FVector2f RotationPivot = FVector2f(0,0);
 		if (TranslationMode == EUVEditorTranslationModeBackend::Absolute)
 		{
-			RotationPivot = GetPivotFromMode(Vid);
+			RotationPivot = GetPivotFromMode(Vid, PivotMode);
 		}				
 		FVector2f UV = FUVEditorUXSettings::InternalUVToExternalUV(ActiveUVLayer->GetElement(Vid));
 		UV = (UV + FVector2f(Translation) - RotationPivot);
 		ActiveUVLayer->SetElement(Vid, FUVEditorUXSettings::ExternalUVToInternalUV(UV));
 	};
 
-	auto ApplyTransformFunc = [&](TFunction<void(int32 Vid)> TransformFunc)
+	auto ApplyTransformFunc = [this](TFunction<void(int32 Vid)> TransformFunc)
 	{
 		RebuildBoundingBoxes();
 
@@ -284,6 +302,15 @@ void FUVEditorUVTransformOp::HandleTransformationOp(FProgressCancel * Progress)
 		}
 	};
 
+	if (!FMath::IsNearlyZero(QuickRotation))
+	{
+		ApplyTransformFunc(QuickRotFunc);
+	}
+	if (!FMath::IsNearlyZero(QuickTranslation.X) || !FMath::IsNearlyZero(QuickTranslation.Y) )
+	{
+		ApplyTransformFunc(QuickTranslateFunc);
+	}
+
 	if (!FMath::IsNearlyEqual(Scale.X, 1.0f) || !FMath::IsNearlyEqual(Scale.Y, 1.0f))
 	{
 		ApplyTransformFunc(ScaleFunc);
@@ -306,7 +333,7 @@ void FUVEditorUVAlignOp::HandleTransformationOp(FProgressCancel* Progress)
 	int32 NumComponents = UVComponents->Num();
 	RebuildBoundingBoxes();
 
-	auto TranslateFunc = [&](const FVector2f& Translation, int32 Vid)
+	auto TranslateFunc = [this](const FVector2f& Translation, int32 Vid)
 	{
 		FVector2f UV = FUVEditorUXSettings::InternalUVToExternalUV(ActiveUVLayer->GetElement(Vid));
 		UV = (UV + Translation);
@@ -470,7 +497,7 @@ void FUVEditorUVDistributeOp::HandleTransformationOp(FProgressCancel* Progress)
 		ensure(false);
 	}
 
-	auto TranslateFunc = [&](const FVector2f& Translation, int32 Vid)
+	auto TranslateFunc = [this](const FVector2f& Translation, int32 Vid)
 	{
 		FVector2f UV = FUVEditorUXSettings::InternalUVToExternalUV(ActiveUVLayer->GetElement(Vid));
 		UV = (UV + Translation);
@@ -519,6 +546,9 @@ TUniquePtr<FDynamicMeshOperator> UUVEditorUVTransformOperatorFactory::MakeNewOpe
 		Op->Rotation = Settings->Rotation;
 		Op->Translation = Settings->Translation;
 
+		Op->QuickTranslation = Settings->QuickTranslation;
+		Op->QuickRotation = Settings->QuickRotation;
+
 		switch (Settings->TranslationMode)
 		{
 		case EUVEditorTranslationMode::Relative:
diff --git a/Engine/Plugins/Experimental/UVEditor/Source/UVEditorTools/Private/UVEditorTransformTool.cpp b/Engine/Plugins/Experimental/UVEditor/Source/UVEditorTools/Private/UVEditorTransformTool.cpp
--- a/Engine/Plugins/Experimental/UVEditor/Source/UVEditorTools/Private/UVEditorTransformTool.cpp
+++ b/Engine/Plugins/Experimental/UVEditor/Source/UVEditorTools/Private/UVEditorTransformTool.cpp
@@ -77,6 +77,14 @@ void UUVEditorTransformTool::Setup()
 	//Settings->bUDIMCVAREnabled = (FUVEditorUXSettings::CVarEnablePrototypeUDIMSupport.GetValueOnGameThread() > 0);
 	AddToolPropertySource(Settings);
 
+	if (ToolMode.Get(EUVEditorUVTransformType::Transform) == EUVEditorUVTransformType::Transform)
+	{
+		QuickTransformSettings = NewObject<UUVEditorUVQuickTransformProperties>(this);
+		QuickTransformSettings->RestoreProperties(this);
+		QuickTransformSettings->Tool = this;
+		AddToolPropertySource(QuickTransformSettings);
+	}
+
 	UContextObjectStore* ContextStore = GetToolManager()->GetContextObjectStore();
 	UVToolSelectionAPI = ContextStore->FindContext<UUVToolSelectionAPI>();
 
@@ -192,9 +200,9 @@ void UUVEditorTransformTool::Shutdown(EToolShutdownType ShutdownType)
 		Target->AppliedPreview->ClearOpFactory();
 	}
 
-	for (int32 TargetIndex = 0; TargetIndex < Targets.Num(); ++TargetIndex)
+	for (int32 FactoryIndex = 0; FactoryIndex < Factories.Num(); ++FactoryIndex)
 	{
-		Factories[TargetIndex] = nullptr;
+		Factories[FactoryIndex] = nullptr;
 	}
 
 	Settings = nullptr;
@@ -209,6 +217,25 @@ void UUVEditorTransformTool::OnTick(float DeltaTime)
 	}
 }
 
+void UUVEditorTransformTool::InitiateQuickTranslate(float Offset, const FVector2D& Direction)
+{
+	Settings->QuickTranslation += Direction * Offset;
+	for (TObjectPtr<UUVEditorToolMeshInput> Target : Targets)
+	{
+		Target->AppliedPreview->InvalidateResult();
+	}
+}
+
+
+void UUVEditorTransformTool::InitiateQuickRotation(float Rotation)
+{
+	Settings->QuickRotation += Rotation;
+	for (TObjectPtr<UUVEditorToolMeshInput> Target : Targets)
+	{
+		Target->AppliedPreview->InvalidateResult();
+	}
+}
+
 void UUVEditorTransformTool::OnPropertyModified(UObject* PropertySet, FProperty* Property)
 {
 	for (TObjectPtr<UUVEditorToolMeshInput> Target : Targets)
diff --git a/Engine/Plugins/Experimental/UVEditor/Source/UVEditorTools/Public/Operators/UVEditorUVTransformOp.h b/Engine/Plugins/Experimental/UVEditor/Source/UVEditorTools/Public/Operators/UVEditorUVTransformOp.h
--- a/Engine/Plugins/Experimental/UVEditor/Source/UVEditorTools/Public/Operators/UVEditorUVTransformOp.h
+++ b/Engine/Plugins/Experimental/UVEditor/Source/UVEditorTools/Public/Operators/UVEditorUVTransformOp.h
@@ -11,6 +11,7 @@
 #include "UVEditorUVTransformOp.generated.h"
 
 class UUVTransformProperties;
+class UUVEditorTransformTool;
 
 /**
  * UV Transform Strategies for the UV Transform Tool
@@ -106,35 +107,42 @@ public:
 	EUVEditorUVTransformType TransformType = EUVEditorUVTransformType::Transform;
 
 	/** Scale applied to UVs, potentially non-uniform */
-	UPROPERTY(EditAnywhere, Category = "Scaling", Transient, meta = (DisplayName = "Scale",
+	UPROPERTY(EditAnywhere, Category = "Advanced Transform | Scaling", DuplicateTransient, meta = (TransientToolProperty, DisplayName = "Scale",
 		      EditCondition = "TransformType == EUVEditorUVTransformType::Transform", EditConditionHides, HideEditConditionToggle = true))
 	FVector2D Scale = FVector2D(1.0, 1.0); 
 
 	/** Rotation applied to UVs after scaling, specified in degrees */
-	UPROPERTY(EditAnywhere, Category = "Rotation", Transient, meta = (DisplayName = "Rotation",
+	UPROPERTY(EditAnywhere, Category = "Advanced Transform | Rotation", meta = (TransientToolProperty, DisplayName = "Rotation",
 		      EditCondition = "TransformType == EUVEditorUVTransformType::Transform", EditConditionHides, HideEditConditionToggle = true))
 	float Rotation = 0.0;
 
 	/** Translation applied to UVs, and after scaling and rotation */
-	UPROPERTY(EditAnywhere, Category = "Translation", Transient, meta = (DisplayName = "Transform",
+	UPROPERTY(EditAnywhere, Category = "Advanced Transform | Translation", meta = (TransientToolProperty, DisplayName = "Translation",
 		      EditCondition = "TransformType == EUVEditorUVTransformType::Transform", EditConditionHides, HideEditConditionToggle = true))
 	FVector2D Translation = FVector2D(0, 0);
 
 	/** Translation applied to UVs, and after scaling and rotation */
-	UPROPERTY(EditAnywhere, Category = "Translation", Transient, meta = (DisplayName = "Transform Mode",
+	UPROPERTY(EditAnywhere, Category = "Advanced Transform | Translation", meta = (TransientToolProperty, DisplayName = "Translation Mode",
 		EditCondition = "TransformType == EUVEditorUVTransformType::Transform", EditConditionHides, HideEditConditionToggle = true))
 	EUVEditorTranslationMode TranslationMode = EUVEditorTranslationMode::Relative;
 
 	/** Transformation origin mode used for scaling and rotation */
-	UPROPERTY(EditAnywhere, Category = "Transform Origin", Transient, meta = (DisplayName = "Mode",
+	UPROPERTY(EditAnywhere, Category = "Advanced Transform | Transform Origin", meta = (TransientToolProperty, DisplayName = "Mode",
 		EditCondition = "TransformType == EUVEditorUVTransformType::Transform", EditConditionHides, HideEditConditionToggle = true))
 	EUVEditorPivotType PivotMode = EUVEditorPivotType::BoundingBoxCenter;
 
 	/** Manual Transformation origin point */
-	UPROPERTY(EditAnywhere, Category = "Transform Origin", Transient, meta = (DisplayName = "Coords",
+	UPROPERTY(EditAnywhere, Category = "Advanced Transform | Transform Origin", meta = (TransientToolProperty, DisplayName = "Coords",
 		EditCondition = "TransformType == EUVEditorUVTransformType::Transform && PivotMode == EUVEditorPivotType::Manual", EditConditionHides, HideEditConditionToggle = true))
 	FVector2D ManualPivot = FVector2D(0, 0);
 
+	UPROPERTY(meta = (TransientToolProperty))
+	FVector2D QuickTranslation;
+
+	UPROPERTY(meta = (TransientToolProperty))
+	float QuickRotation;
+
+
 	/** Controls what geometry the alignment is to be relative to when performed. */
 	UPROPERTY(EditAnywhere, Category = "Align", meta = (DisplayName = "Alignment Anchor",
 		EditCondition = "TransformType == EUVEditorUVTransformType::Align", EditConditionHides, HideEditConditionToggle = true))
@@ -154,7 +162,6 @@ public:
 	UPROPERTY(EditAnywhere, Category = "Distribute", meta = (DisplayName = "Distribution Mode",
 		EditCondition = "TransformType == EUVEditorUVTransformType::Distribute", EditConditionHides, HideEditConditionToggle = true))
 	EUVEditorDistributeMode DistributeMode = EUVEditorDistributeMode::TopEdges;
-
 };
 
 
@@ -257,6 +264,9 @@ public:
 	virtual ~FUVEditorUVTransformOp() {}
 
 	// inputs
+	FVector2D QuickTranslation = FVector2D(0.0, 0.0);
+	float QuickRotation = 0.0;
+
 	FVector2D Scale = FVector2D(1.0, 1.0);
 	float Rotation = 0.0;
 	FVector2D Translation = FVector2D(0, 0);
@@ -267,7 +277,7 @@ public:
 
 protected:
 	virtual void HandleTransformationOp(FProgressCancel* Progress) override;
-	FVector2f GetPivotFromMode(int32 ElementID);
+	FVector2f GetPivotFromMode(int32 ElementID, EUVEditorPivotTypeBackend Mode);
 };
 
 class UVEDITORTOOLS_API FUVEditorUVAlignOp : public FUVEditorUVTransformBaseOp
diff --git a/Engine/Plugins/Experimental/UVEditor/Source/UVEditorTools/Public/UVEditorTransformTool.h b/Engine/Plugins/Experimental/UVEditor/Source/UVEditorTools/Public/UVEditorTransformTool.h
--- a/Engine/Plugins/Experimental/UVEditor/Source/UVEditorTools/Public/UVEditorTransformTool.h
+++ b/Engine/Plugins/Experimental/UVEditor/Source/UVEditorTools/Public/UVEditorTransformTool.h
@@ -62,6 +62,28 @@ protected:
 	virtual void ConfigureTool(UUVEditorTransformTool* NewTool) const  override;
 };
 
+/**
+ * UV Quick Transform Settings
+ */
+UCLASS()
+class UVEDITORTOOLS_API UUVEditorUVQuickTransformProperties : public UInteractiveToolPropertySet
+{
+	GENERATED_BODY()
+
+public:
+	UPROPERTY(EditAnywhere, Category = "Quick Transform", meta = (TransientToolProperty, DisplayName = "Translation"))
+	float QuickTranslateOffset = 0.0;
+
+	UPROPERTY(EditAnywhere, Category = "Quick Transform", meta = (TransientToolProperty, DisplayName = "Rotation"))
+	float QuickRotationOffset = 0.0;
+
+	// parent ref required for details customization
+	UPROPERTY(meta = (TransientToolProperty))
+	TWeakObjectPtr<UUVEditorTransformTool> Tool;
+
+};
+
+
 /**
  * 
  */
@@ -94,6 +116,10 @@ public:
 	virtual void DrawHUD(FCanvas* Canvas, IToolsContextRenderAPI* RenderAPI) override;
 
 	virtual void OnPropertyModified(UObject* PropertySet, FProperty* Property) override;
+
+	void InitiateQuickTranslate(float Offset, const FVector2D& Direction);
+	void InitiateQuickRotation(float Rotation);
+
 protected:	
 
 	TOptional<EUVEditorUVTransformType> ToolMode;
@@ -101,6 +127,9 @@ protected:
 	UPROPERTY()
 	TArray<TObjectPtr<UUVEditorToolMeshInput>> Targets;
 
+	UPROPERTY()
+	TObjectPtr<UUVEditorUVQuickTransformProperties> QuickTransformSettings = nullptr;
+
 	UPROPERTY()
 	TObjectPtr<UUVEditorUVTransformProperties> Settings = nullptr;
 
diff --git a/Engine/Plugins/Experimental/UVEditor/Source/UVEditorToolsEditorOnly/Private/DetailsCustomizations/UVTransformToolCustomizations.cpp b/Engine/Plugins/Experimental/UVEditor/Source/UVEditorToolsEditorOnly/Private/DetailsCustomizations/UVTransformToolCustomizations.cpp
new file mode 100644
--- /dev/null
+++ b/Engine/Plugins/Experimental/UVEditor/Source/UVEditorToolsEditorOnly/Private/DetailsCustomizations/UVTransformToolCustomizations.cpp
@@ -0,0 +1,286 @@
+// Copyright Epic Games, Inc. All Rights Reserved.
+
+#include "DetailsCustomizations/UVTransformToolCustomizations.h"
+#include "DetailCategoryBuilder.h"
+#include "DetailLayoutBuilder.h"
+#include "DetailWidgetRow.h"
+#include "PropertyCustomizationHelpers.h"
+#include "Widgets/Input/SComboButton.h"
+#include "Widgets/Layout/SUniformGridPanel.h"
+#include "IDetailChildrenBuilder.h"
+#include "Internationalization/BreakIterator.h"
+#include "Framework/MultiBox/MultiBoxBuilder.h"
+#include "IContentBrowserSingleton.h"
+#include "SAssetView.h"
+
+#include "PropertyRestriction.h"
+
+#include "UVEditorTransformTool.h"
+
+
+#define LOCTEXT_NAMESPACE "UVEditorDetailsCustomization"
+
+namespace UVEditorDetailsCustomizationLocal
+{
+	void CustomSortTransformToolCategories(const TMap<FName, IDetailCategoryBuilder*>&  AllCategoryMap )
+	{
+		(*AllCategoryMap.Find(FName("Quick Translate")))->SetSortOrder(0);
+		(*AllCategoryMap.Find(FName("Quick Rotate")))->SetSortOrder(1);
+		(*AllCategoryMap.Find(FName("Quick Transform")))->SetSortOrder(2);
+	}
+}
+
+
+//
+// UVEditorTransformTool
+//
+
+
+TSharedRef<IDetailCustomization> FUVEditorUVTransformToolDetails::MakeInstance()
+{
+	return MakeShareable(new FUVEditorUVTransformToolDetails);
+}
+
+
+void FUVEditorUVTransformToolDetails::CustomizeDetails(IDetailLayoutBuilder& DetailBuilder)
+{
+	using namespace UVEditorDetailsCustomizationLocal;
+
+	TArray<TWeakObjectPtr<UObject>> ObjectsBeingCustomized;
+	DetailBuilder.GetObjectsBeingCustomized(ObjectsBeingCustomized);
+	check(ObjectsBeingCustomized.Num() > 0);
+	UUVEditorUVQuickTransformProperties* QuickTransformProperties = CastChecked<UUVEditorUVQuickTransformProperties>(ObjectsBeingCustomized[0]);
+	UUVEditorTransformTool* Tool = QuickTransformProperties->Tool.Get();
+	TargetTool = Tool;
+
+	BuildQuickTranslateMenu(DetailBuilder);
+	BuildQuickRotateMenu(DetailBuilder);
+
+	DetailBuilder.SortCategories(&CustomSortTransformToolCategories);
+}
+
+void FUVEditorUVTransformToolDetails::BuildQuickTranslateMenu(IDetailLayoutBuilder& DetailBuilder)
+{
+	const FName QuickTranslateCategoryName = "Quick Translate";
+	IDetailCategoryBuilder& CategoryBuilder = DetailBuilder.EditCategory(QuickTranslateCategoryName);	
+
+	TSharedPtr<IPropertyHandle> QuickTranslateOffsetHandle = DetailBuilder.GetProperty(GET_MEMBER_NAME_CHECKED(UUVEditorUVQuickTransformProperties, QuickTranslateOffset), UUVEditorUVQuickTransformProperties::StaticClass());
+	ensure(QuickTranslateOffsetHandle->IsValidHandle());
+	QuickTranslateOffsetHandle->MarkHiddenByCustomization();
+
+	auto ApplyTranslation = [this, QuickTranslateOffsetHandle](const FVector2D& Direction)
+	{
+		float TranslationValue;
+		QuickTranslateOffsetHandle->GetValue(TranslationValue);
+		return OnQuickTranslate(TranslationValue, Direction);
+	};
+
+	FDetailWidgetRow& CustomRow = CategoryBuilder.AddCustomRow(FText::GetEmpty())
+	[
+		SNew(SUniformGridPanel)
+		.SlotPadding(FMargin(5.0f))
+		+ SUniformGridPanel::Slot(0, 0)
+		.HAlign(HAlign_Center)
+		.VAlign(VAlign_Center)
+		[
+			SNew(SButton)
+			.Text(LOCTEXT("QuickMoveTopLeft", "TL"))
+			.ToolTipText(LOCTEXT("QuickMoveTopLeftToolTip", "Applies translation offset in the negative X axis and the positive Y axis"))
+			.OnClicked_Lambda([ApplyTranslation]() { return ApplyTranslation(FVector2D(-1.0, 1.0)); })			
+		]
+		+ SUniformGridPanel::Slot(1, 0)
+		.HAlign(HAlign_Center)
+		.VAlign(VAlign_Center)
+		[
+			SNew(SButton)
+			.Text(LOCTEXT("QuickMoveTop", "T"))
+			.ToolTipText(LOCTEXT("QuickMoveTopToolTip", "Applies the translation offset in the positive Y axis."))
+			.OnClicked_Lambda([ApplyTranslation]() { return ApplyTranslation(FVector2D(0.0, 1.0)); })
+		]
+		+ SUniformGridPanel::Slot(2, 0)
+		.HAlign(HAlign_Center)
+		.VAlign(VAlign_Center)
+		[
+			SNew(SButton)
+			.Text(LOCTEXT("QuickMoveTopRight", "TR"))
+			.ToolTipText(LOCTEXT("QuickMoveTopRightToolTip", "Applies the translation offset in the positive X axis and the positive Y axis."))
+			.OnClicked_Lambda([ApplyTranslation]() { return ApplyTranslation(FVector2D(1.0, 1.0)); })
+		]
+		+ SUniformGridPanel::Slot(0, 1)
+		.HAlign(HAlign_Center)
+		.VAlign(VAlign_Center)
+		[
+			SNew(SButton)
+			.Text(LOCTEXT("QuickMoveLeft", "L"))
+			.ToolTipText(LOCTEXT("QuickMoveLeftToolTip", "Applies translation offset in the negative X axis."))
+			.OnClicked_Lambda([ApplyTranslation]() { return ApplyTranslation(FVector2D(-1.0, 0.0)); })
+		]
+		+ SUniformGridPanel::Slot(1, 1)
+		.HAlign(HAlign_Center)
+		.VAlign(VAlign_Center)
+		[
+			SNew(SProperty, QuickTranslateOffsetHandle)
+			.ShouldDisplayName(false)
+		]
+		+ SUniformGridPanel::Slot(2, 1)
+		.HAlign(HAlign_Center)
+		.VAlign(VAlign_Center)
+		[
+			SNew(SButton)
+			.Text(LOCTEXT("QuickMoveRight", "R"))
+			.ToolTipText(LOCTEXT("QuickMoveRightToolTip", "Applies the translation offset in the positive X axis."))
+			.OnClicked_Lambda([ApplyTranslation]() { return ApplyTranslation(FVector2D(1.0, 0.0)); })
+		]
+		+ SUniformGridPanel::Slot(0, 2)
+		.HAlign(HAlign_Center)
+		.VAlign(VAlign_Center)
+		[
+			SNew(SButton)
+			.Text(LOCTEXT("QuickMoveBottomLeft", "BL"))
+			.ToolTipText(LOCTEXT("QuickMoveBottomLeftToolTip", "Applies translation offset in the negative X axis and the negative Y axis"))
+			.OnClicked_Lambda([ApplyTranslation]() { return ApplyTranslation(FVector2D(-1.0, -1.0)); })
+		]
+		+ SUniformGridPanel::Slot(1, 2)
+		.HAlign(HAlign_Center)
+		.VAlign(VAlign_Center)
+		[
+			SNew(SButton)
+			.Text(LOCTEXT("QuickMoveBottom", "B"))
+			.ToolTipText(LOCTEXT("QuickMoveBottomToolTip", "Applies the translation offset in the negative Y axis."))
+			.OnClicked_Lambda([ApplyTranslation]() { return ApplyTranslation(FVector2D(0.0, -1.0)); })
+		]
+		+ SUniformGridPanel::Slot(2, 2)
+		.HAlign(HAlign_Center)
+		.VAlign(VAlign_Center)
+		[
+			SNew(SButton)
+			.Text(LOCTEXT("QuickMoveBottomRight", "BR"))
+			.ToolTipText(LOCTEXT("QuickMoveBottomRightToolTip", "Applies the translation offset in the positive X axis and the negative Y axis."))
+			.OnClicked_Lambda([ApplyTranslation]() { return ApplyTranslation(FVector2D(1.0, -1.0)); })
+		]
+
+	];
+
+}
+
+
+void FUVEditorUVTransformToolDetails::BuildQuickRotateMenu(IDetailLayoutBuilder& DetailBuilder)
+{
+	const FName QuickRotateCategoryName = "Quick Rotate";
+	IDetailCategoryBuilder& CategoryBuilder = DetailBuilder.EditCategory(QuickRotateCategoryName);
+
+	TSharedPtr<IPropertyHandle> QuickRotationOffsetHandle = DetailBuilder.GetProperty(GET_MEMBER_NAME_CHECKED(UUVEditorUVQuickTransformProperties, QuickRotationOffset), UUVEditorUVQuickTransformProperties::StaticClass());
+	ensure(QuickRotationOffsetHandle->IsValidHandle());
+	QuickRotationOffsetHandle->MarkHiddenByCustomization();
+
+	FDetailWidgetRow& CustomRow = CategoryBuilder.AddCustomRow(FText::GetEmpty())
+	[
+		SNew(SUniformGridPanel)
+		.SlotPadding(FMargin(5.0f))
+		+ SUniformGridPanel::Slot(0, 0)
+		.HAlign(HAlign_Center)
+		.VAlign(VAlign_Center)
+		[
+			SNew(SButton)
+			.Text(LOCTEXT("QuickRotateClockwise", "CW"))
+			.ToolTipText(LOCTEXT("QuickRotateClockwiseToolTip", "Applies the rotation in a clockwise orientation"))
+			.OnClicked_Lambda([this, QuickRotationOffsetHandle]() {
+				float RotationValue;
+				QuickRotationOffsetHandle->GetValue(RotationValue);
+				return OnQuickRotate(-RotationValue); 
+			})
+		]
+		+ SUniformGridPanel::Slot(1, 0)
+		.HAlign(HAlign_Center)
+		.VAlign(VAlign_Center)
+		[
+			SNew(SProperty, QuickRotationOffsetHandle)
+			.ShouldDisplayName(false)
+		]
+		+ SUniformGridPanel::Slot(2, 0)
+		.HAlign(HAlign_Center)
+		.VAlign(VAlign_Center)
+		[
+			SNew(SButton)
+			.Text(LOCTEXT("QuickRotateCounterclockwise", "CCW"))
+			.ToolTipText(LOCTEXT("QuickRotateCounterclockwiseToolTip", "Applies the rotation in a counter clockwise orientation"))
+			.OnClicked_Lambda([this, QuickRotationOffsetHandle]() {
+				float RotationValue;
+				QuickRotationOffsetHandle->GetValue(RotationValue);
+				return OnQuickRotate(RotationValue); 
+			})
+		]
+
+		+ SUniformGridPanel::Slot(0, 1)
+		.HAlign(HAlign_Center)
+		.VAlign(VAlign_Center)
+		[
+			SNew(SButton)
+			.Text(LOCTEXT("QuickRotateClockwise10Deg", "10°"))
+			.ToolTipText(LOCTEXT("QuickRotateClockwise10DegToolTip", "Applies a 10 degree clockwise orientation"))
+			.OnClicked_Lambda([this]() { return OnQuickRotate(-10.0); })
+		]
+		+ SUniformGridPanel::Slot(2, 1)
+		.HAlign(HAlign_Center)
+		.VAlign(VAlign_Center)
+		[
+			SNew(SButton)
+			.Text(LOCTEXT("QuickRotateCounterclockwise10Deg", "10°"))
+			.ToolTipText(LOCTEXT("QuickRotateCounterclockwise10DegToolTip", "Applies a 10 degree counter clockwise orientation"))
+			.OnClicked_Lambda([this]() { return OnQuickRotate(10.0); })
+		]
+		+ SUniformGridPanel::Slot(0, 2)
+		.HAlign(HAlign_Center)
+		.VAlign(VAlign_Center)
+		[
+			SNew(SButton)
+			.Text(LOCTEXT("QuickRotateClockwise45Deg", "45°"))
+			.ToolTipText(LOCTEXT("QuickRotateClockwise45DegToolTip", "Applies a 45 degree clockwise orientation"))
+			.OnClicked_Lambda([this]() { return OnQuickRotate(-45.0); })
+		]
+		+ SUniformGridPanel::Slot(2, 2)
+		.HAlign(HAlign_Center)
+		.VAlign(VAlign_Center)
+		[
+			SNew(SButton)
+			.Text(LOCTEXT("QuickRotateCounterclockwise45Deg", "45°"))
+			.ToolTipText(LOCTEXT("QuickRotateCounterclockwise45DegToolTip", "Applies a 45 degree counter clockwise orientation"))
+			.OnClicked_Lambda([this]() { return OnQuickRotate(45.0); })
+		]
+		+ SUniformGridPanel::Slot(0, 3)
+		.HAlign(HAlign_Center)
+		.VAlign(VAlign_Center)
+		[
+			SNew(SButton)
+			.Text(LOCTEXT("QuickRotateClockwise90Deg", "90°"))
+			.ToolTipText(LOCTEXT("QuickRotateClockwise90DegToolTip", "Applies a 90 degree clockwise orientation"))
+			.OnClicked_Lambda([this]() { return OnQuickRotate(-90.0); })
+		]
+		+ SUniformGridPanel::Slot(2, 3)
+		.HAlign(HAlign_Center)
+		.VAlign(VAlign_Center)
+		[
+			SNew(SButton)
+			.Text(LOCTEXT("QuickRotateCounterclockwise90Deg", "90°"))
+			.ToolTipText(LOCTEXT("QuickRotateCounterclockwise90DegToolTip", "Applies a 90 degree counter clockwise orientation"))
+			.OnClicked_Lambda([this]() { return OnQuickRotate(90.0); })
+		]
+	];
+}
+
+
+FReply  FUVEditorUVTransformToolDetails::OnQuickTranslate(float TranslationValue, const FVector2D& Direction)
+{
+	ensure(TargetTool.IsValid());
+	TargetTool->InitiateQuickTranslate(TranslationValue, Direction);
+	return FReply::Handled();
+}
+
+FReply  FUVEditorUVTransformToolDetails::OnQuickRotate(float RotationValue)
+{
+	ensure(TargetTool.IsValid());
+	TargetTool->InitiateQuickRotation(RotationValue);
+	return FReply::Handled();
+}
+
+#undef LOCTEXT_NAMESPACE
\ No newline at end of file
diff --git a/Engine/Plugins/Experimental/UVEditor/Source/UVEditorToolsEditorOnly/Private/UVEditorToolsEditorOnlyModule.cpp b/Engine/Plugins/Experimental/UVEditor/Source/UVEditorToolsEditorOnly/Private/UVEditorToolsEditorOnlyModule.cpp
--- a/Engine/Plugins/Experimental/UVEditor/Source/UVEditorToolsEditorOnly/Private/UVEditorToolsEditorOnlyModule.cpp
+++ b/Engine/Plugins/Experimental/UVEditor/Source/UVEditorToolsEditorOnly/Private/UVEditorToolsEditorOnlyModule.cpp
@@ -6,8 +6,10 @@
 #include "PropertyEditorModule.h"
 
 #include "DetailsCustomizations/UVUnwrapToolCustomizations.h"
+#include "DetailsCustomizations/UVTransformToolCustomizations.h"
 
 #include "Operators/UVEditorRecomputeUVsOp.h"
+#include "UVEditorTransformTool.h"
 
 #define LOCTEXT_NAMESPACE "FUVEditorToolsEditorOnlyModule"
 
@@ -51,6 +53,11 @@ void FUVEditorToolsEditorOnlyModule::OnPostEngineInit()
 	/// Unwrap
 	PropertyModule.RegisterCustomClassLayout("UVEditorRecomputeUVsToolProperties", FOnGetDetailCustomizationInstance::CreateStatic(&FUVEditorRecomputeUVsToolDetails::MakeInstance));
 	ClassesToUnregisterOnShutdown.Add(UUVEditorRecomputeUVsToolProperties::StaticClass()->GetFName());
+
+	// Transform
+	PropertyModule.RegisterCustomClassLayout("UVEditorUVQuickTransformProperties", FOnGetDetailCustomizationInstance::CreateStatic(&FUVEditorUVTransformToolDetails::MakeInstance));
+	ClassesToUnregisterOnShutdown.Add(UUVEditorUVQuickTransformProperties::StaticClass()->GetFName());
+
 }
 
 #undef LOCTEXT_NAMESPACE
diff --git a/Engine/Plugins/Experimental/UVEditor/Source/UVEditorToolsEditorOnly/Public/DetailsCustomizations/UVTransformToolCustomizations.h b/Engine/Plugins/Experimental/UVEditor/Source/UVEditorToolsEditorOnly/Public/DetailsCustomizations/UVTransformToolCustomizations.h
new file mode 100644
--- /dev/null
+++ b/Engine/Plugins/Experimental/UVEditor/Source/UVEditorToolsEditorOnly/Public/DetailsCustomizations/UVTransformToolCustomizations.h
@@ -0,0 +1,28 @@
+// Copyright Epic Games, Inc. All Rights Reserved.
+
+#pragma once
+
+#include "IDetailCustomization.h"
+#include "Templates/SharedPointer.h"
+
+class IPropertyHandle;
+class FPropertyRestriction;
+class UUVEditorTransformTool;
+
+// Customization for UUVEditorUVTransformProperties
+class FUVEditorUVTransformToolDetails : public IDetailCustomization
+{
+public:
+	static TSharedRef<IDetailCustomization> MakeInstance();
+	virtual void CustomizeDetails(IDetailLayoutBuilder& DetailBuilder) override;
+
+protected:
+	TWeakObjectPtr<UUVEditorTransformTool> TargetTool;
+	void BuildQuickTranslateMenu(IDetailLayoutBuilder& DetailBuilder);
+	void BuildQuickRotateMenu(IDetailLayoutBuilder& DetailBuilder);
+
+	FReply OnQuickTranslate(float TranslationValue, const FVector2D& Direction);
+	FReply OnQuickRotate(float RotationValue);
+};
+
+
