UVEditor: Correct error with lambda capture causing CI failures.

#rnx

diff --git a/Engine/Plugins/Experimental/UVEditor/Source/UVEditorTools/Private/Operators/UVEditorUVTransformOp.cpp b/Engine/Plugins/Experimental/UVEditor/Source/UVEditorTools/Private/Operators/UVEditorUVTransformOp.cpp
--- a/Engine/Plugins/Experimental/UVEditor/Source/UVEditorTools/Private/Operators/UVEditorUVTransformOp.cpp
+++ b/Engine/Plugins/Experimental/UVEditor/Source/UVEditorTools/Private/Operators/UVEditorUVTransformOp.cpp
@@ -272,7 +272,7 @@ void FUVEditorUVTransformOp::HandleTransformationOp(FProgressCancel * Progress)
 		ActiveUVLayer->SetElement(Vid, FUVEditorUXSettings::ExternalUVToInternalUV(UV));
 	};
 
-	auto ApplyTransformFunc = [this](TFunction<void(int32 Vid)> TransformFunc)
+	auto ApplyTransformFunc = [this, &Progress](TFunction<void(int32 Vid)> TransformFunc)
 	{
 		RebuildBoundingBoxes();
 
