Fix non-unity
#rnx

diff --git a/Engine/Plugins/Experimental/UVEditor/Source/UVEditorToolsEditorOnly/Private/DetailsCustomizations/UVTransformToolCustomizations.cpp b/Engine/Plugins/Experimental/UVEditor/Source/UVEditorToolsEditorOnly/Private/DetailsCustomizations/UVTransformToolCustomizations.cpp
--- a/Engine/Plugins/Experimental/UVEditor/Source/UVEditorToolsEditorOnly/Private/DetailsCustomizations/UVTransformToolCustomizations.cpp
+++ b/Engine/Plugins/Experimental/UVEditor/Source/UVEditorToolsEditorOnly/Private/DetailsCustomizations/UVTransformToolCustomizations.cpp
@@ -5,6 +5,7 @@
 #include "DetailLayoutBuilder.h"
 #include "DetailWidgetRow.h"
 #include "PropertyCustomizationHelpers.h"
+#include "Widgets/Input/SButton.h"
 #include "Widgets/Input/SComboButton.h"
 #include "Widgets/Layout/SUniformGridPanel.h"
 #include "IDetailChildrenBuilder.h"
diff --git a/Engine/Plugins/Experimental/UVEditor/Source/UVEditorToolsEditorOnly/Public/DetailsCustomizations/UVTransformToolCustomizations.h b/Engine/Plugins/Experimental/UVEditor/Source/UVEditorToolsEditorOnly/Public/DetailsCustomizations/UVTransformToolCustomizations.h
--- a/Engine/Plugins/Experimental/UVEditor/Source/UVEditorToolsEditorOnly/Public/DetailsCustomizations/UVTransformToolCustomizations.h
+++ b/Engine/Plugins/Experimental/UVEditor/Source/UVEditorToolsEditorOnly/Public/DetailsCustomizations/UVTransformToolCustomizations.h
@@ -3,7 +3,9 @@
 #pragma once
 
 #include "IDetailCustomization.h"
+#include "Input/Reply.h"
 #include "Templates/SharedPointer.h"
+#include "UObject/WeakObjectPtr.h"
 
 class IPropertyHandle;
 class FPropertyRestriction;
diff --git a/Engine/Source/Runtime/Renderer/Private/PathTracingSpatialTemporalDenoising.cpp b/Engine/Source/Runtime/Renderer/Private/PathTracingSpatialTemporalDenoising.cpp
--- a/Engine/Source/Runtime/Renderer/Private/PathTracingSpatialTemporalDenoising.cpp
+++ b/Engine/Source/Runtime/Renderer/Private/PathTracingSpatialTemporalDenoising.cpp
@@ -4,6 +4,7 @@
 #include "PathTracing.h"
 #include "RHI.h"
 #include "PathTracingDenoiser.h"
+#include "RHIDefinitions.h"
 
 PathTracingDenoiserFunction* GPathTracingDenoiserFunc = nullptr;
 
