UVEditor: Move logic for selection type conversions to the FUVToolSelection object. This allows for getting equivalent selection representations without actively needing to change the current selection itself.

diff --git a/Engine/Plugins/Experimental/UVEditor/Source/UVEditorTools/Private/Selection/UVEditorMeshSelectionMechanic.cpp b/Engine/Plugins/Experimental/UVEditor/Source/UVEditorTools/Private/Selection/UVEditorMeshSelectionMechanic.cpp
--- a/Engine/Plugins/Experimental/UVEditor/Source/UVEditorTools/Private/Selection/UVEditorMeshSelectionMechanic.cpp
+++ b/Engine/Plugins/Experimental/UVEditor/Source/UVEditorTools/Private/Selection/UVEditorMeshSelectionMechanic.cpp
@@ -975,192 +975,17 @@ void UUVEditorMeshSelectionMechanic::SetSelectionMode(
 	const TArray<FUVToolSelection>& OriginalSelections = SelectionAPI->GetSelections();
 	const TArray<FUVToolSelection>& OriginalUnsetSelections = SelectionAPI->GetUnsetElementAppliedMeshSelections();
 
-	auto ConvertSelections = [](bool bUnsetSelections, const TArray<FUVToolSelection>& OriginalSelections, FUVToolSelection::EType CurrentSelectionType, FUVToolSelection::EType ExpectedSelectionType)
+	TArray<FUVToolSelection> NewSelections, NewUnsetSelections;
+	NewSelections.Reserve(OriginalSelections.Num());
+	NewUnsetSelections.Reserve(OriginalUnsetSelections.Num());
+	for (const FUVToolSelection& OriginalSelection : OriginalSelections)
 	{
-		TArray<FUVToolSelection> NewSelections;
-		for (const FUVToolSelection& OriginalSelection : OriginalSelections)
-		{
-			FDynamicMesh3* Mesh = bUnsetSelections ? OriginalSelection.Target->AppliedCanonical.Get(): OriginalSelection.Target->UnwrapCanonical.Get();
-
-			FUVToolSelection NewSelection;
-			check(CurrentSelectionType == OriginalSelection.Type);
-			NewSelection.Target = OriginalSelection.Target;
-			NewSelection.Type = ExpectedSelectionType;
-
-			const TSet<int32> OriginalIDs = OriginalSelection.SelectedIDs;
-			TSet<int32>& NewIDs = NewSelection.SelectedIDs;
-
-			auto VerticesToEdges = [Mesh, &OriginalIDs, &NewIDs]()
-			{
-				TRACE_CPUPROFILER_EVENT_SCOPE(MeshSelectionMechanic_ChangeSelectionMode_VerticesToEdges);
-
-				for (int32 Vid : OriginalIDs)
-				{
-					for (int32 Eid : Mesh->VtxEdgesItr(Vid))
-					{
-						if (!NewIDs.Contains(Eid))
-						{
-							FIndex2i Verts = Mesh->GetEdgeV(Eid);
-							if (OriginalIDs.Contains(Verts.A) &&
-								OriginalIDs.Contains(Verts.B))
-							{
-								NewIDs.Add(Eid);
-							}
-						}
-					}
-				}
-			};
-
-			auto VerticesToTriangles = [Mesh, &OriginalIDs, &NewIDs]()
-			{
-				TRACE_CPUPROFILER_EVENT_SCOPE(MeshSelectionMechanic_ChangeSelectionMode_VerticesToTriangles);
-
-				for (int32 Vid : OriginalIDs)
-				{
-					for (int32 Tid : Mesh->VtxTrianglesItr(Vid))
-					{
-						if (!NewIDs.Contains(Tid))
-						{
-							FIndex3i Verts = Mesh->GetTriangle(Tid);
-							if (OriginalIDs.Contains(Verts.A) &&
-								OriginalIDs.Contains(Verts.B) &&
-								OriginalIDs.Contains(Verts.C))
-							{
-								NewIDs.Add(Tid);
-							}
-						}
-					}
-				}
-			};
-
-			auto EdgesToVertices = [Mesh, &OriginalIDs, &NewIDs]()
-			{
-				TRACE_CPUPROFILER_EVENT_SCOPE(MeshSelectionMechanic_ChangeSelectionMode_EdgesToVertices);
-
-				for (int32 Eid : OriginalIDs)
-				{
-					FIndex2i Verts = Mesh->GetEdgeV(Eid);
-					NewIDs.Add(Verts.A);
-					NewIDs.Add(Verts.B);
-				}
-			};
-
-			// Triangles with two selected edges will be selected
-			auto EdgesToTriangles = [Mesh, &OriginalIDs, &NewIDs]()
-			{
-				TRACE_CPUPROFILER_EVENT_SCOPE(MeshSelectionMechanic_ChangeSelectionMode_EdgesToTriangles);
-
-				TArray<int32> FoundTriangles;
-				for (int32 Eid : OriginalIDs)
-				{
-					FIndex2i Tris = Mesh->GetEdgeT(Eid);
-					FoundTriangles.Add(Tris.A);
-					if (Tris.B != IndexConstants::InvalidID)
-					{
-						FoundTriangles.Add(Tris.B);
-					}
-				}
-
-				if (FoundTriangles.Num() < 2)
-				{
-					return;
-				}
-
-				Algo::Sort(FoundTriangles);
-
-				for (int I = 0; I < FoundTriangles.Num() - 1; I++)
-				{
-					if (FoundTriangles[I] == FoundTriangles[I + 1])
-					{
-						NewIDs.Add(FoundTriangles[I]);
-						I++;
-					}
-				}
-			};
-
-			auto TrianglesToVertices = [Mesh, &OriginalIDs, &NewIDs]()
-			{
-				TRACE_CPUPROFILER_EVENT_SCOPE(MeshSelectionMechanic_ChangeSelectionMode_TrianglesToVertices);
-
-				for (int32 Tid : OriginalIDs)
-				{
-					FIndex3i Verts = Mesh->GetTriangle(Tid);
-					NewIDs.Add(Verts.A);
-					NewIDs.Add(Verts.B);
-					NewIDs.Add(Verts.C);
-				}
-			};
-
-			auto TrianglesToEdges = [Mesh, &OriginalIDs, &NewIDs]()
-			{
-				TRACE_CPUPROFILER_EVENT_SCOPE(MeshSelectionMechanic_ChangeSelectionMode_TrianglesToEdges);
-
-				for (int32 Tid : OriginalIDs)
-				{
-					FIndex3i Edges = Mesh->GetTriEdgesRef(Tid);
-					NewIDs.Add(Edges.A);
-					NewIDs.Add(Edges.B);
-					NewIDs.Add(Edges.C);
-				}
-			};
-
-			switch (CurrentSelectionType)
-			{
-			case FUVToolSelection::EType::Vertex:
-				switch (ExpectedSelectionType)
-				{
-				case FUVToolSelection::EType::Vertex:
-					ensure(false); // Should have been an early-out
-					break;
-				case FUVToolSelection::EType::Edge:
-					VerticesToEdges();
-					break;
-				case FUVToolSelection::EType::Triangle:
-					VerticesToTriangles();
-					break;
-				}
-				break; // SelectionMode
-			case FUVToolSelection::EType::Edge:
-				switch (ExpectedSelectionType)
-				{
-				case FUVToolSelection::EType::Vertex:
-					EdgesToVertices();
-					break;
-				case FUVToolSelection::EType::Edge:
-					ensure(false); // Should have been an early-out
-					break;
-				case FUVToolSelection::EType::Triangle:
-					EdgesToTriangles();
-					break;
-				}
-				break; // SelectionMode
-			case FUVToolSelection::EType::Triangle:
-				switch (ExpectedSelectionType)
-				{
-				case FUVToolSelection::EType::Vertex:
-					TrianglesToVertices();
-					break;
-				case FUVToolSelection::EType::Edge:
-					TrianglesToEdges();
-					break;
-				case FUVToolSelection::EType::Triangle:
-					ensure(false); // Should have been an early-out
-					break;
-				}
-				break; // SelectionMode
-			}
-
-			if (!NewSelection.IsEmpty())
-			{
-				NewSelections.Add(MoveTemp(NewSelection));
-			}
-		}
-
-		return NewSelections;
-	};
-
-	TArray<FUVToolSelection> NewSelections = ConvertSelections(false, OriginalSelections,CurrentSelectionType, ExpectedSelectionType);
-	TArray<FUVToolSelection> NewUnsetSelections = ConvertSelections(true, OriginalUnsetSelections, CurrentSelectionType, ExpectedSelectionType);
+		NewSelections.Add(OriginalSelection.GetConvertedSelection(*OriginalSelection.Target->UnwrapCanonical, ExpectedSelectionType));
+	}
+	for (const FUVToolSelection& OriginalUnsetSelection : OriginalUnsetSelections)
+	{
+		NewUnsetSelections.Add(OriginalUnsetSelection.GetConvertedSelection(*OriginalUnsetSelection.Target->AppliedCanonical, ExpectedSelectionType));
+	}
 
 	// Apply selection change
 	SelectionAPI->BeginChange();
diff --git a/Engine/Plugins/Experimental/UVEditor/Source/UVEditorTools/Private/Selection/UVToolSelection.cpp b/Engine/Plugins/Experimental/UVEditor/Source/UVEditorTools/Private/Selection/UVToolSelection.cpp
--- a/Engine/Plugins/Experimental/UVEditor/Source/UVEditorTools/Private/Selection/UVToolSelection.cpp
+++ b/Engine/Plugins/Experimental/UVEditor/Source/UVEditorTools/Private/Selection/UVToolSelection.cpp
@@ -57,4 +57,181 @@ bool FUVToolSelection::AreElementsPresentInMesh(FDynamicMesh3& Mesh) const
 	return true;
 }
 
+FUVToolSelection FUVToolSelection::GetConvertedSelection(const FDynamicMesh3& Mesh, FUVToolSelection::EType ExpectedSelectionType) const
+{
+	if (Type == ExpectedSelectionType)
+	{
+		return *this;
+	}
+
+	FUVToolSelection NewSelection;
+	NewSelection.Target = Target;
+	NewSelection.Type = ExpectedSelectionType;
+
+	const TSet<int32>& OriginalIDs = SelectedIDs;
+	TSet<int32>& NewIDs = NewSelection.SelectedIDs;
+
+	auto VerticesToEdges = [&Mesh, &OriginalIDs, &NewIDs]()
+	{
+		TRACE_CPUPROFILER_EVENT_SCOPE(ToolSelection_GetConvertedSelection_VerticesToEdges);
+
+		for (int32 Vid : OriginalIDs)
+		{
+			for (int32 Eid : Mesh.VtxEdgesItr(Vid))
+			{
+				if (!NewIDs.Contains(Eid))
+				{
+					FIndex2i Verts = Mesh.GetEdgeV(Eid);
+					if (OriginalIDs.Contains(Verts.A) &&
+						OriginalIDs.Contains(Verts.B))
+					{
+						NewIDs.Add(Eid);
+					}
+				}
+			}
+		}
+	};
+
+	auto VerticesToTriangles = [&Mesh, &OriginalIDs, &NewIDs]()
+	{
+		TRACE_CPUPROFILER_EVENT_SCOPE(ToolSelection_GetConvertedSelection_VerticesToTriangles);
+
+		for (int32 Vid : OriginalIDs)
+		{
+			for (int32 Tid : Mesh.VtxTrianglesItr(Vid))
+			{
+				if (!NewIDs.Contains(Tid))
+				{
+					FIndex3i Verts = Mesh.GetTriangle(Tid);
+					if (OriginalIDs.Contains(Verts.A) &&
+						OriginalIDs.Contains(Verts.B) &&
+						OriginalIDs.Contains(Verts.C))
+					{
+						NewIDs.Add(Tid);
+					}
+				}
+			}
+		}
+	};
+
+	auto EdgesToVertices = [&Mesh, &OriginalIDs, &NewIDs]()
+	{
+		TRACE_CPUPROFILER_EVENT_SCOPE(ToolSelection_GetConvertedSelection_EdgesToVertices);
+
+		for (int32 Eid : OriginalIDs)
+		{
+			FIndex2i Verts = Mesh.GetEdgeV(Eid);
+			NewIDs.Add(Verts.A);
+			NewIDs.Add(Verts.B);
+		}
+	};
+
+	// Triangles with two selected edges will be selected
+	auto EdgesToTriangles = [&Mesh, &OriginalIDs, &NewIDs]()
+	{
+		TRACE_CPUPROFILER_EVENT_SCOPE(ToolSelection_GetConvertedSelection_EdgesToTriangles);
+
+		TArray<int32> FoundTriangles;
+		for (int32 Eid : OriginalIDs)
+		{
+			FIndex2i Tris = Mesh.GetEdgeT(Eid);
+			FoundTriangles.Add(Tris.A);
+			if (Tris.B != IndexConstants::InvalidID)
+			{
+				FoundTriangles.Add(Tris.B);
+			}
+		}
+
+		if (FoundTriangles.Num() < 2)
+		{
+			return;
+		}
+
+		Algo::Sort(FoundTriangles);
+
+		for (int I = 0; I < FoundTriangles.Num() - 1; I++)
+		{
+			if (FoundTriangles[I] == FoundTriangles[I + 1])
+			{
+				NewIDs.Add(FoundTriangles[I]);
+				I++;
+			}
+		}
+	};
+
+	auto TrianglesToVertices = [&Mesh, &OriginalIDs, &NewIDs]()
+	{
+		TRACE_CPUPROFILER_EVENT_SCOPE(ToolSelection_GetConvertedSelection_TrianglesToVertices);
+
+		for (int32 Tid : OriginalIDs)
+		{
+			FIndex3i Verts = Mesh.GetTriangle(Tid);
+			NewIDs.Add(Verts.A);
+			NewIDs.Add(Verts.B);
+			NewIDs.Add(Verts.C);
+		}
+	};
+
+	auto TrianglesToEdges = [Mesh, &OriginalIDs, &NewIDs]()
+	{
+		TRACE_CPUPROFILER_EVENT_SCOPE(ToolSelection_GetConvertedSelection_TrianglesToEdges);
+
+		for (int32 Tid : OriginalIDs)
+		{
+			FIndex3i Edges = Mesh.GetTriEdgesRef(Tid);
+			NewIDs.Add(Edges.A);
+			NewIDs.Add(Edges.B);
+			NewIDs.Add(Edges.C);
+		}
+	};
+
+	switch (Type)
+	{
+	case FUVToolSelection::EType::Vertex:
+		switch (ExpectedSelectionType)
+		{
+		case FUVToolSelection::EType::Vertex:
+			ensure(false); // Should have been an early-out
+			break;
+		case FUVToolSelection::EType::Edge:
+			VerticesToEdges();
+			break;
+		case FUVToolSelection::EType::Triangle:
+			VerticesToTriangles();
+			break;
+		}
+		break; // SelectionMode
+	case FUVToolSelection::EType::Edge:
+		switch (ExpectedSelectionType)
+		{
+		case FUVToolSelection::EType::Vertex:
+			EdgesToVertices();
+			break;
+		case FUVToolSelection::EType::Edge:
+			ensure(false); // Should have been an early-out
+			break;
+		case FUVToolSelection::EType::Triangle:
+			EdgesToTriangles();
+			break;
+		}
+		break; // SelectionMode
+	case FUVToolSelection::EType::Triangle:
+		switch (ExpectedSelectionType)
+		{
+		case FUVToolSelection::EType::Vertex:
+			TrianglesToVertices();
+			break;
+		case FUVToolSelection::EType::Edge:
+			TrianglesToEdges();
+			break;
+		case FUVToolSelection::EType::Triangle:
+			ensure(false); // Should have been an early-out
+			break;
+		}
+		break; // SelectionMode
+	}
+	
+	return NewSelection;
+}
+
 #undef LOCTEXT_NAMESPACE
\ No newline at end of file
diff --git a/Engine/Plugins/Experimental/UVEditor/Source/UVEditorTools/Public/Selection/UVToolSelection.h b/Engine/Plugins/Experimental/UVEditor/Source/UVEditorTools/Public/Selection/UVToolSelection.h
--- a/Engine/Plugins/Experimental/UVEditor/Source/UVEditorTools/Public/Selection/UVToolSelection.h
+++ b/Engine/Plugins/Experimental/UVEditor/Source/UVEditorTools/Public/Selection/UVToolSelection.h
@@ -66,6 +66,11 @@ public:
 
 	bool AreElementsPresentInMesh(FDynamicMesh3& Mesh) const;
 
+	/**
+	 * Utility method to construct a new selection from an existing reference selection, only with a new geometry type.
+	 */
+	FUVToolSelection GetConvertedSelection(const FDynamicMesh3& Mesh, FUVToolSelection::EType ExpectedSelectionType) const;
+
 protected:
 	
 	FMeshEdgesFromTriangleSubIndices StableEdgeIDs;
