UVEditor: Add additional static convenience methods for determining UDIMs from points and bounding boxes in UV space.

diff --git a/Engine/Plugins/Experimental/UVEditor/Source/UVEditorTools/Private/Utilities/MeshUDIMClassifier.cpp b/Engine/Plugins/Experimental/UVEditor/Source/UVEditorTools/Private/Utilities/MeshUDIMClassifier.cpp
--- a/Engine/Plugins/Experimental/UVEditor/Source/UVEditorTools/Private/Utilities/MeshUDIMClassifier.cpp
+++ b/Engine/Plugins/Experimental/UVEditor/Source/UVEditorTools/Private/Utilities/MeshUDIMClassifier.cpp
@@ -42,8 +42,7 @@ FVector2i FDynamicMeshUDIMClassifier::ClassifyTrianglesToUDIM(const FDynamicMesh
 		UVOverlay->GetTriElements(Tid, Vertex0, Vertex1, Vertex2);
 		FVector2f BaryCenter = (Vertex0 + Vertex1 + Vertex2) / 3.0f;
 		BaryCenter = FUVEditorUXSettings::InternalUVToExternalUV(BaryCenter);
-		UDIM.X = FMath::Floor(BaryCenter.X);
-		UDIM.Y = FMath::Floor(BaryCenter.Y);
+		UDIM = ClassifyPointToUDIM(BaryCenter);		
 		return UDIM;
 	};
 
@@ -66,6 +65,40 @@ FVector2i FDynamicMeshUDIMClassifier::ClassifyTrianglesToUDIM(const FDynamicMesh
 	return MaximumUDIM;
 }
 
+FVector2i FDynamicMeshUDIMClassifier::ClassifyBoundingBoxToUDIM(const FDynamicMeshUVOverlay* UVOverlay, const FAxisAlignedBox2d& BoundingBox)
+{
+	FVector2i MinUDIM = ClassifyPointToUDIM((FVector2f)BoundingBox.Min);
+	double MaxOverlapArea = 0.0;
+	FVector2i SelectedUDIM = MinUDIM;
+
+	// We only really need to test the minimum UDIM tile and one greater in each dimension.
+	// Either the BoundingBox will be fully overlapping in any or none of these, but
+	// we don't need to find the largest indexed UDIM with full overlap, just the first.
+	for (int32 UDIM_X = MinUDIM.X; UDIM_X <= MinUDIM.X+1; ++UDIM_X)
+	{
+		for (int32 UDIM_Y = MinUDIM.Y; UDIM_Y <= MinUDIM.Y+1; ++UDIM_Y)
+		{
+			FAxisAlignedBox2d TestUDIMBoundingBox(FVector2d(UDIM_X, UDIM_Y), FVector2d(UDIM_X + 1, UDIM_Y + 1));
+			double OverlapArea = BoundingBox.Intersect(TestUDIMBoundingBox).Area();
+			if (OverlapArea > MaxOverlapArea)
+			{
+				SelectedUDIM = FVector2i(UDIM_X, UDIM_Y);
+				MaxOverlapArea = OverlapArea;
+			}
+		}
+	}
+
+	return SelectedUDIM;
+}
+
+
+FVector2i FDynamicMeshUDIMClassifier::ClassifyPointToUDIM(const FVector2f& UVPoint)
+{
+	FVector2i UDIM;
+	UDIM.X = FMath::FloorToInt32(UVPoint.X);
+	UDIM.Y = FMath::FloorToInt32(UVPoint.Y);
+	return UDIM;
+}
 
 void FDynamicMeshUDIMClassifier::ClassifyUDIMs()
 {
diff --git a/Engine/Plugins/Experimental/UVEditor/Source/UVEditorTools/Public/Utilities/MeshUDIMClassifier.h b/Engine/Plugins/Experimental/UVEditor/Source/UVEditorTools/Public/Utilities/MeshUDIMClassifier.h
--- a/Engine/Plugins/Experimental/UVEditor/Source/UVEditorTools/Public/Utilities/MeshUDIMClassifier.h
+++ b/Engine/Plugins/Experimental/UVEditor/Source/UVEditorTools/Public/Utilities/MeshUDIMClassifier.h
@@ -26,6 +26,8 @@ namespace UE
 			TArray<int32> TidsForTile(FVector2i TileIndexIn) const;
 
 			static FVector2i ClassifyTrianglesToUDIM(const FDynamicMeshUVOverlay* UVOverlay, TArray<int32> Tids);
+			static FVector2i ClassifyBoundingBoxToUDIM(const FDynamicMeshUVOverlay* UVOverlay, const FAxisAlignedBox2d& BoundingBox);
+			static FVector2i ClassifyPointToUDIM(const FVector2f& UVPoint);
 
 		protected:
 
