UVEditor: Transform/Align/Distribute tool update.

* Tools now support vertex and edge selections, allowing for all operations to apply on a per element basis
* Align and Distribute now support a per component / per selection element toggle, allowing for fine scale control over operations
* Distribute now supports a new overlap removal option, which attempts to minimally move components apart until they don't overlap
* Created details customizations for new translate/distribute/align tool functionality and added additional styling data.

diff --git a/Engine/Plugins/Experimental/UVEditor/Source/UVEditor/Private/UVEditorStyle.cpp b/Engine/Plugins/Experimental/UVEditor/Source/UVEditor/Private/UVEditorStyle.cpp
deleted file mode 100644
diff --git a/Engine/Plugins/Experimental/UVEditor/Source/UVEditor/Public/UVEditorStyle.h b/Engine/Plugins/Experimental/UVEditor/Source/UVEditor/Public/UVEditorStyle.h
deleted file mode 100644
diff --git a/Engine/Plugins/Experimental/UVEditor/Source/UVEditorTools/Private/Operators/UVEditorUVTransformOp.cpp b/Engine/Plugins/Experimental/UVEditor/Source/UVEditorTools/Private/Operators/UVEditorUVTransformOp.cpp
--- a/Engine/Plugins/Experimental/UVEditor/Source/UVEditorTools/Private/Operators/UVEditorUVTransformOp.cpp
+++ b/Engine/Plugins/Experimental/UVEditor/Source/UVEditorTools/Private/Operators/UVEditorUVTransformOp.cpp
@@ -12,11 +12,59 @@
 
 using namespace UE::Geometry;
 
+namespace TransformOpLocals
+{
+	FVector2f UnwrapPositionToUV(const FVector3d& Position)
+	{
+		return FUVEditorUXSettings::ExternalUVToInternalUV(FUVEditorUXSettings::VertPositionToUV(Position));
+	}
+
+	FVector3d UVToUnwrapPosition(const FVector2f& UV)
+	{
+		return FUVEditorUXSettings::UVToVertPosition(FUVEditorUXSettings::InternalUVToExternalUV(UV));
+	}
+}
+
 void FUVEditorUVTransformBaseOp::SetTransform(const FTransformSRT3d& Transform)
 {
 	ResultTransform = Transform;
 }
 
+void FUVEditorUVTransformBaseOp::SetSelection(const TSet<int32>& TriangleSelectionIn, const TSet<int32>& VertexSelectionIn)
+{
+		TriangleSelection = TriangleSelectionIn;
+		VertexSelection = VertexSelectionIn;
+}
+
+void FUVEditorUVTransformBaseOp::SortComponentsByBoundingBox()
+{
+	SpatiallyOrderedComponentIndex.SetNum(UVComponents->Num());
+	for (int32 Index = 0; Index < UVComponents->Num(); ++Index)
+	{
+		SpatiallyOrderedComponentIndex[Index] = Index;
+	}
+	TArrayView<int32> IndexView(SpatiallyOrderedComponentIndex.GetData(), SpatiallyOrderedComponentIndex.Num());
+	TArrayView<FAxisAlignedBox2d> BBView(PerComponentBoundingBoxes.GetData(), PerComponentBoundingBoxes.Num());
+
+	auto BBIndexSort = [&BBView](const int32& A, const int32& B)
+	{
+		FVector2D CenterDiff = BBView[A].Center() - BBView[B].Center();
+
+		if (CenterDiff.X < 0)
+		{
+			return true;
+		}
+		else if (FMath::IsNearlyZero(CenterDiff.X) && CenterDiff.Y < 0)
+		{
+			return true;
+		}
+		return false;
+	};
+
+	IndexView.StableSort(BBIndexSort);
+}
+
+
 void FUVEditorUVTransformBaseOp::RebuildBoundingBoxes()
 {	
 	int32 NumComponents = UVComponents->Num();
@@ -25,31 +73,25 @@ void FUVEditorUVTransformBaseOp::RebuildBoundingBoxes()
 
 	ParallelFor(NumComponents, [&](int32 k)
 		{
-			FVector2d TriangleUV[3];
+			FVector3d VertexPos[3];
 			PerComponentBoundingBoxes[k] = FAxisAlignedBox2d::Empty();
-			const TArray<int>& Triangles = (*UVComponents)[k].Indices;
-			for (int tid : Triangles)
-			{
-				FIndex3i TriElements = ActiveUVLayer->GetTriangle(tid);
-				if (!TriElements.Contains(FDynamicMesh3::InvalidID))
-				{
-					for (int j = 0; j < 3; ++j)
-					{
-						TriangleUV[j] = FVector2d(FUVEditorUXSettings::InternalUVToExternalUV(ActiveUVLayer->GetElement(TriElements[j])));
-						PerComponentBoundingBoxes[k].Contain(TriangleUV[j]);
-					}
-				}
+			const TArray<int>& Vertices = (*UVComponents)[k].Indices;
+			for (int32 Vid : Vertices)
+			{				
+				PerComponentBoundingBoxes[k].Contain((FVector2d)TransformOpLocals::UnwrapPositionToUV(ResultMesh->GetVertexRef(Vid)));
 			}
 		});
+
 	for (int32 Cid = 0; Cid < NumComponents; ++Cid)
 	{
 		OverallBoundingBox.Contain(PerComponentBoundingBoxes[Cid]);
 	}
+
 }
 
 void FUVEditorUVTransformBaseOp::CollectTransformElements()
 {
-	if (Selection.IsSet())
+	if (VertexSelection.IsSet())
 	{
 		TransformingElements = TSet<int32>();
 	}
@@ -58,18 +100,14 @@ void FUVEditorUVTransformBaseOp::CollectTransformElements()
 
 	for (int32 k = 0; k < NumComponents; ++k)
 	{
-		const TArray<int>& Triangles = (*UVComponents)[k].Indices;
-		for (int tid : Triangles)
+		const TArray<int>& Vertices = (*UVComponents)[k].Indices;
+		for (int Vid : Vertices)
 		{
-			FIndex3i Elements = ActiveUVLayer->GetTriangle(tid);
-			for (int j = 0; j < 3; ++j)
+			if (VertexSelection.IsSet() && VertexSelection.GetValue().Contains(Vid))
 			{
-				if (TransformingElements.IsSet())
-				{
-					TransformingElements.GetValue().Add(Elements[j]);
-				}
-				ElementToComponent.Add(Elements[j], k);
+				TransformingElements.GetValue().Add(Vid);
 			}
+			ElementToComponent.Add(Vid, k);
 		}
 	}
 }
@@ -80,6 +118,9 @@ FVector2f FUVEditorUVTransformBaseOp::GetAlignmentPointFromBoundingBoxAndDirecti
 
 	switch (Direction)
 	{
+	case EUVEditorAlignDirectionBackend::None:
+		AlignmentPoint = FVector2f(0,0);
+		break;
 	case EUVEditorAlignDirectionBackend::Top:
 		AlignmentPoint = FVector2f(BoundingBox.Center().X, BoundingBox.Max.Y);
 		break;
@@ -113,6 +154,9 @@ FVector2f FUVEditorUVTransformBaseOp::GetAlignmentPointFromUDIMAndDirection(EUVE
 
 	switch (Direction)
 	{
+	case EUVEditorAlignDirectionBackend::None:
+		AlignmentPoint = FVector2f(0, 0);
+		break;
 	case EUVEditorAlignDirectionBackend::Top:
 		AlignmentPoint = FVector2f(UDIMLowerCorner.X + 0.5f, UDIMLowerCorner.Y + 1.0f);
 		break;
@@ -161,17 +205,38 @@ void FUVEditorUVTransformBaseOp::CalculateResult(FProgressCancel* Progress)
 
 	auto UVIslandPredicate = [this](int32 Triangle0, int32 Triangle1)
 	{
-		return ActiveUVLayer->AreTrianglesConnected(Triangle0, Triangle1);
+		return GroupingMode != EUVEditorAlignDistributeGroupingModeBackend::IndividualVertices;
 	};
 
 	UVComponents = MakeShared<FMeshConnectedComponents>(ResultMesh.Get());
-	if (Selection.IsSet())
+	if(GroupingMode == EUVEditorAlignDistributeGroupingModeBackend::EnclosingBoundingBox)
 	{
-		UVComponents->FindConnectedTriangles(Selection.GetValue().Array(), UVIslandPredicate);
+		// If we're using the enclosing bounding box, we'll just "fake" one big component and proceed as normal.
+		UVComponents->Components.Add( new FMeshConnectedComponents::FComponent());
+		if (VertexSelection.IsSet())
+		{
+			UVComponents->Components[0].Indices = VertexSelection.GetValue().Array();
+		}
+		else
+		{
+			UVComponents->Components[0].Indices.Reserve(ResultMesh.Get()->VertexCount());
+			for (int32 Vid : ResultMesh.Get()->VertexIndicesItr())
+			{
+				UVComponents->Components[0].Indices.Add(Vid);
+			}
+		}
+
 	}
 	else
 	{
-		UVComponents->FindConnectedTriangles(UVIslandPredicate);
+		if (VertexSelection.IsSet())
+		{
+			UVComponents->FindConnectedVertices(VertexSelection.GetValue().Array(), UVIslandPredicate);
+		}
+		else
+		{
+			UVComponents->FindConnectedVertices(UVIslandPredicate);
+		}
 	}
 
 	if (Progress && Progress->Cancelled())
@@ -221,24 +286,24 @@ void FUVEditorUVTransformOp::HandleTransformationOp(FProgressCancel * Progress)
 	auto ScaleFunc = [this](int32 Vid)
 	{
 		FVector2f ScalePivot = GetPivotFromMode(Vid, PivotMode);
-		FVector2f UV = FUVEditorUXSettings::InternalUVToExternalUV(ActiveUVLayer->GetElement(Vid));
+		FVector2f UV = TransformOpLocals::UnwrapPositionToUV(ResultMesh->GetVertexRef(Vid));
 		UV = (UV - ScalePivot);
 		UV[0] *= Scale[0];
 		UV[1] *= Scale[1];
 		UV = (UV + ScalePivot);
-		ActiveUVLayer->SetElement(Vid, FUVEditorUXSettings::ExternalUVToInternalUV(UV));
+		ResultMesh->SetVertex(Vid, TransformOpLocals::UVToUnwrapPosition(UV));
 	};
 
 	auto BaseRotFunc = [this](int32 Vid, float RotationIn, const FVector2f& Pivot)
 	{
-		FVector2f UV = FUVEditorUXSettings::InternalUVToExternalUV(ActiveUVLayer->GetElement(Vid));
+		FVector2f UV = TransformOpLocals::UnwrapPositionToUV(ResultMesh->GetVertexRef(Vid));
 		FVector2f UV_Rotated;
 		double RotationInRadians = RotationIn / 180.0 * UE_PI;
 		UV = (UV - Pivot);
 		UV_Rotated[0] = UV[0] * FMath::Cos(RotationInRadians) - UV[1] * FMath::Sin(RotationInRadians);
 		UV_Rotated[1] = UV[0] * FMath::Sin(RotationInRadians) + UV[1] * FMath::Cos(RotationInRadians);
 		UV = (UV_Rotated + Pivot);
-		ActiveUVLayer->SetElement(Vid, FUVEditorUXSettings::ExternalUVToInternalUV(UV));
+		ResultMesh->SetVertex(Vid, TransformOpLocals::UVToUnwrapPosition(UV));
 	};
 
 	auto QuickRotFunc = [this, &BaseRotFunc](int32 Vid)
@@ -255,9 +320,9 @@ void FUVEditorUVTransformOp::HandleTransformationOp(FProgressCancel * Progress)
 
 	auto QuickTranslateFunc = [this](int32 Vid)
 	{
-		FVector2f UV = FUVEditorUXSettings::InternalUVToExternalUV(ActiveUVLayer->GetElement(Vid));
+		FVector2f UV = TransformOpLocals::UnwrapPositionToUV(ResultMesh->GetVertexRef(Vid));
 		UV = (UV + FVector2f(QuickTranslation));
-		ActiveUVLayer->SetElement(Vid, FUVEditorUXSettings::ExternalUVToInternalUV(UV));
+		ResultMesh->SetVertex(Vid, TransformOpLocals::UVToUnwrapPosition(UV));
 	};
 
 	auto TranslateFunc = [this](int32 Vid)
@@ -267,9 +332,9 @@ void FUVEditorUVTransformOp::HandleTransformationOp(FProgressCancel * Progress)
 		{
 			RotationPivot = GetPivotFromMode(Vid, PivotMode);
 		}				
-		FVector2f UV = FUVEditorUXSettings::InternalUVToExternalUV(ActiveUVLayer->GetElement(Vid));
+		FVector2f UV = TransformOpLocals::UnwrapPositionToUV(ResultMesh->GetVertexRef(Vid));
 		UV = (UV + FVector2f(Translation) - RotationPivot);
-		ActiveUVLayer->SetElement(Vid, FUVEditorUXSettings::ExternalUVToInternalUV(UV));
+		ResultMesh->SetVertex(Vid, TransformOpLocals::UVToUnwrapPosition(UV));
 	};
 
 	auto ApplyTransformFunc = [this, &Progress](TFunction<void(int32 Vid)> TransformFunc)
@@ -290,7 +355,7 @@ void FUVEditorUVTransformOp::HandleTransformationOp(FProgressCancel * Progress)
 		}
 		else
 		{
-			for (int ElementID : ActiveUVLayer->ElementIndicesItr())
+			for (int ElementID : ResultMesh->VertexIndicesItr())
 			{
 				TransformFunc(ElementID);
 
@@ -335,15 +400,17 @@ void FUVEditorUVAlignOp::HandleTransformationOp(FProgressCancel* Progress)
 
 	auto TranslateFunc = [this](const FVector2f& Translation, int32 Vid)
 	{
-		FVector2f UV = FUVEditorUXSettings::InternalUVToExternalUV(ActiveUVLayer->GetElement(Vid));
+		FVector2f UV = TransformOpLocals::UnwrapPositionToUV(ResultMesh->GetVertexRef(Vid));
 		UV = (UV + Translation);
-		ActiveUVLayer->SetElement(Vid, FUVEditorUXSettings::ExternalUVToInternalUV(UV));
+		ResultMesh->SetVertex(Vid, TransformOpLocals::UVToUnwrapPosition(UV));
 	};
 
 	auto TranslationFromAlignmentPoints = [this](const FVector2f& PointTo, const FVector2f& PointFrom)
 	{
 		switch (AlignDirection)
 		{
+		case EUVEditorAlignDirectionBackend::None:
+			return FVector2f(0, 0);
 		case EUVEditorAlignDirectionBackend::Top:
 			return FVector2f(0, PointTo.Y - PointFrom.Y);
 		case EUVEditorAlignDirectionBackend::Bottom:
@@ -353,52 +420,74 @@ void FUVEditorUVAlignOp::HandleTransformationOp(FProgressCancel* Progress)
 		case EUVEditorAlignDirectionBackend::Right:
 			return FVector2f(PointTo.X - PointFrom.X, 0);
 		case EUVEditorAlignDirectionBackend::CenterVertically:
-			return FVector2f(PointTo.X - PointFrom.X, 0);
-		case EUVEditorAlignDirectionBackend::CenterHorizontally:
 			return FVector2f(0, PointTo.Y - PointFrom.Y);
+		case EUVEditorAlignDirectionBackend::CenterHorizontally:
+			return FVector2f(PointTo.X - PointFrom.X, 0);			
 		default:
 			ensure(false);
 			return FVector2f(0, 0);
 		}
 	};
 
-	TArray<FVector2f> PerComponentTranslation;
-	PerComponentTranslation.SetNum(NumComponents);
-	for (int32 Cid = 0; Cid < NumComponents; ++Cid)
+	auto ComponentToUDIM = [this](int32 ComponentID)
+	{
+		return FDynamicMeshUDIMClassifier::ClassifyBoundingBoxToUDIM(ActiveUVLayer, PerComponentBoundingBoxes[ComponentID]);
+	};
+	auto UVElementToUDIM = [this](int32 Vid)
+	{
+		return FDynamicMeshUDIMClassifier::ClassifyPointToUDIM(TransformOpLocals::UnwrapPositionToUV(ResultMesh->GetVertexRef(Vid)));
+	};
+	auto ComponentToPoint = [this](int32 ComponentID)
 	{
-		FVector2f ComponentAlignmentPoint = GetAlignmentPointFromBoundingBoxAndDirection(AlignDirection, PerComponentBoundingBoxes[Cid]);
-		switch (AlignAnchor)
+		return GetAlignmentPointFromBoundingBoxAndDirection(AlignDirection, PerComponentBoundingBoxes[ComponentID]);
+	};
+	auto UVElementToPoint = [this](int32 Vid)
+	{
+		return TransformOpLocals::UnwrapPositionToUV(ResultMesh->GetVertexRef(Vid));
+	};
+
+
+	auto GetNeededTranslateOffset = [this, TranslationFromAlignmentPoints](int32 PointID, EUVEditorAlignAnchorBackend AlignAnchorMode,
+		TFunction<FVector2i(int32 ID)> PointIDToUDIM, TFunction<FVector2f(int32 ID)> PointIDToPoint)
+	{
+		FVector2f PointToAlign = PointIDToPoint(PointID);
+		switch (AlignAnchorMode)
 		{
 		case EUVEditorAlignAnchorBackend::UDIMTile:
 		{
-			FVector2i UDIM = FDynamicMeshUDIMClassifier::ClassifyTrianglesToUDIM(ActiveUVLayer, UVComponents->Components[Cid].Indices);
+			FVector2i UDIM = PointIDToUDIM(PointID);
 			FVector2f TileAlignmentPoint = GetAlignmentPointFromUDIMAndDirection(AlignDirection, UDIM);
-			PerComponentTranslation[Cid] = TranslationFromAlignmentPoints(TileAlignmentPoint, ComponentAlignmentPoint);
+			return TranslationFromAlignmentPoints(TileAlignmentPoint, PointToAlign);
 		}
-			break;
+		break;
 		case EUVEditorAlignAnchorBackend::BoundingBox:
-		{		
+		{
 			FVector2f BoundingBoxAlignmentPoint = GetAlignmentPointFromBoundingBoxAndDirection(AlignDirection, OverallBoundingBox);
-			PerComponentTranslation[Cid] = TranslationFromAlignmentPoints(BoundingBoxAlignmentPoint, ComponentAlignmentPoint);
+			return TranslationFromAlignmentPoints(BoundingBoxAlignmentPoint, PointToAlign);
 		}
 		break;
 		case EUVEditorAlignAnchorBackend::Manual:
 		{
-			PerComponentTranslation[Cid] = TranslationFromAlignmentPoints((FVector2f)ManualAnchor, ComponentAlignmentPoint);
+			return TranslationFromAlignmentPoints((FVector2f)ManualAnchor, PointToAlign);
 		}
 		break;
 		default:
 			ensure(false);
-			PerComponentTranslation[Cid] = FVector2f();
+			return FVector2f();
 			break;
 		}
-	}
+	};
 
 	if (TransformingElements.IsSet())
 	{
 		for (int ElementID : TransformingElements.GetValue())
 		{
-			TranslateFunc(PerComponentTranslation[*ElementToComponent.Find(ElementID)], ElementID);
+			int32* ComponentID = ElementToComponent.Find(ElementID);
+			if (ensure(ComponentID))
+			{
+				FVector2f TranslateOffset = GetNeededTranslateOffset(*ComponentID, AlignAnchor, ComponentToUDIM, ComponentToPoint);
+				TranslateFunc(TranslateOffset, ElementID);
+			}
 			if (Progress && Progress->Cancelled())
 			{
 				return;
@@ -409,7 +498,13 @@ void FUVEditorUVAlignOp::HandleTransformationOp(FProgressCancel* Progress)
 	{
 		for (int ElementID : ActiveUVLayer->ElementIndicesItr())
 		{
-			TranslateFunc(PerComponentTranslation[*ElementToComponent.Find(ElementID)], ElementID);
+
+			int32* ComponentID = ElementToComponent.Find(ElementID);
+			if (ensure(ComponentID))
+			{
+				FVector2f TranslateOffset = GetNeededTranslateOffset(*ComponentID, AlignAnchor, ComponentToUDIM, ComponentToPoint);
+				TranslateFunc(TranslateOffset, ElementID);
+			}
 			if (Progress && Progress->Cancelled())
 			{
 				return;
@@ -425,6 +520,7 @@ void FUVEditorUVDistributeOp::HandleTransformationOp(FProgressCancel* Progress)
 {
 	int32 NumComponents = UVComponents->Num();
 	RebuildBoundingBoxes();
+	SortComponentsByBoundingBox();
 	TArray<FVector2f> PerComponentTranslation;
 	
 
@@ -438,12 +534,20 @@ void FUVEditorUVDistributeOp::HandleTransformationOp(FProgressCancel* Progress)
 			TotalDistance += bVertical ? PerComponentBoundingBoxes[Cid].Height() : PerComponentBoundingBoxes[Cid].Width();
 		}
 		float BoundingBoxDistance = bVertical ? OverallBoundingBox.Height() : OverallBoundingBox.Width();
-		float GapSpace = (BoundingBoxDistance - TotalDistance) / (NumComponents-1);
-		float PerComponentSpace = BoundingBoxDistance / (NumComponents);
+		if (bEnableManualDistances)
+		{
+			BoundingBoxDistance = ManualExtent;
+		}
+		float GapSpace = FMath::Max(0, (BoundingBoxDistance - TotalDistance) / (NumComponents - 1));
+		if(bEnableManualDistances && bEqualizeSpacing)
+		{
+			GapSpace = ManualSpacing;
+		}
+		float PerComponentSpace = BoundingBoxDistance / (NumComponents - 1);
 
 		float NextPosition = 0.0f;
 		FVector2f OverallAlignmentPoint = GetAlignmentPointFromBoundingBoxAndDirection(EdgeDirection, OverallBoundingBox);
-		for (int32 Cid = 0; Cid < NumComponents; ++Cid)
+		for (int32 Cid: SpatiallyOrderedComponentIndex)
 		{
 			FVector2f ComponentAlignmentPoint = GetAlignmentPointFromBoundingBoxAndDirection(EdgeDirection, PerComponentBoundingBoxes[Cid]);
 			if (bVertical)
@@ -467,41 +571,117 @@ void FUVEditorUVDistributeOp::HandleTransformationOp(FProgressCancel* Progress)
 		return PerComponentTranslation;
 	};
 
+	auto ComputeMinimallyRemoveOverlap = [this, NumComponents]()
+	{
+		FVector2D Center = OverallBoundingBox.Center();
+
+		// Create a local copy so we can move things around...
+		TArray<FAxisAlignedBox2d> ComponentBoundingBoxes = PerComponentBoundingBoxes;
+		// This could be tunable, if we think it needs to be, or perhaps based off bounding box size somehow?
+		double StepAmount = 0.01;
+
+		if (bEnableManualDistances)
+		{
+			for (int32 Cid = 0; Cid < NumComponents; ++Cid)
+			{
+				ComponentBoundingBoxes[Cid].Expand(ManualSpacing / 2.0);
+			}
+		}
+
+		bool bAnyBoxesIntersecting = false;
+		// The following loop will iteratively move all overlapping bounding boxes away from each other and from the center.
+		// This will eventually terminate, by "exploding" all boxes away from the center point in the worst case. If a box
+		// doesn't overlap, it won't move
+		// TODO: Implement a closer fit version that looks at real overlaps between island elements instead of just bounding boxes
+		// TODO: Implement an improvement that uses clustering to prevent far away islands from moving too much away from a global center
+		do
+		{
+			bAnyBoxesIntersecting = false;
+			for (int32 Cid = 0; Cid < NumComponents; ++Cid) 
+			{
+				FVector2D Adjustment = FVector2D::Zero();
+				bool bIsOverlapping = false;
+				for (int32 CidOverlap = 0; CidOverlap < NumComponents; ++CidOverlap)
+				{
+					if (Cid == CidOverlap)
+					{
+						continue;
+					}
+					// Move away from overlapping boxes
+					if (ComponentBoundingBoxes[Cid].Intersects(ComponentBoundingBoxes[CidOverlap]))
+					{
+						FVector2D Offset = (ComponentBoundingBoxes[Cid].Center() - ComponentBoundingBoxes[CidOverlap].Center());
+						Offset.Normalize();
+						if(Offset.IsNearlyZero())
+						{
+							// Edge case where two bounding boxes have identically overlapping center points.
+							// We need to break them up so they don't move in concert
+							// Our "emergency" separation vector places the components deterministically
+							// around a circle so they all follow different paths.
+							Offset.X = FMath::Cos((2 * UE_PI / NumComponents) * Cid);
+							Offset.Y = FMath::Sin((2 * UE_PI / NumComponents) * Cid);							
+						}
+						Adjustment = Adjustment + (Offset * StepAmount);
+						bIsOverlapping = true;
+					}
+				}
+				ComponentBoundingBoxes[Cid].Min += Adjustment;
+				ComponentBoundingBoxes[Cid].Max += Adjustment;
+				bAnyBoxesIntersecting |= bIsOverlapping;
+			}
+		} while (bAnyBoxesIntersecting);
+
+		// Once all bounding boxes have been moved, record their movement to the translation array return value
+		TArray<FVector2f> PerComponentTranslation;
+		PerComponentTranslation.SetNum(NumComponents);
+		for (int32 Cid = 0; Cid < NumComponents; ++Cid)
+		{
+			PerComponentTranslation[Cid] = (FVector2f)(ComponentBoundingBoxes[Cid].Center() - PerComponentBoundingBoxes[Cid].Center());
+		}
+		return PerComponentTranslation;
+	};
+
 	switch (DistributeMode)
 	{
-	case EUVEditorDistributeModeBackend::TopEdges:
-		PerComponentTranslation = ComputeDistributeTranslations(true, EUVEditorAlignDirectionBackend::Top, -1.0f, false);
-		break;
-	case EUVEditorDistributeModeBackend::BottomEdges:
-		PerComponentTranslation = ComputeDistributeTranslations(true, EUVEditorAlignDirectionBackend::Bottom, 1.0f, false);
-		break;
-	case EUVEditorDistributeModeBackend::LeftEdges:
-		PerComponentTranslation = ComputeDistributeTranslations(false, EUVEditorAlignDirectionBackend::Left, 1.0f, false);
-		break;
-	case EUVEditorDistributeModeBackend::RightEdges:
-		PerComponentTranslation = ComputeDistributeTranslations(false, EUVEditorAlignDirectionBackend::Right, -1.0f, false);
-		break;
-	case EUVEditorDistributeModeBackend::CentersVertically:
-		PerComponentTranslation = ComputeDistributeTranslations(false, EUVEditorAlignDirectionBackend::CenterVertically, 1.0f, false);
-		break;
-	case EUVEditorDistributeModeBackend::CentersHorizontally:
-		PerComponentTranslation = ComputeDistributeTranslations(true, EUVEditorAlignDirectionBackend::CenterHorizontally, 1.0f, false);
-		break;
-	case EUVEditorDistributeModeBackend::HorizontalSpace:
-		PerComponentTranslation = ComputeDistributeTranslations(false, EUVEditorAlignDirectionBackend::Left, 1.0f, true);
-		break;
-	case EUVEditorDistributeModeBackend::VerticalSpace:
-		PerComponentTranslation = ComputeDistributeTranslations(true, EUVEditorAlignDirectionBackend::Bottom, 1.0f, true);
-		break;
-	default:
-		ensure(false);
+		case EUVEditorDistributeModeBackend::None:
+			PerComponentTranslation.Init(FVector2f::ZeroVector, NumComponents);
+			break;
+		case EUVEditorDistributeModeBackend::TopEdges:
+			PerComponentTranslation = ComputeDistributeTranslations(true, EUVEditorAlignDirectionBackend::Top, -1.0f, false);
+			break;
+		case EUVEditorDistributeModeBackend::BottomEdges:
+			PerComponentTranslation = ComputeDistributeTranslations(true, EUVEditorAlignDirectionBackend::Bottom, 1.0f, false);
+			break;
+		case EUVEditorDistributeModeBackend::LeftEdges:
+			PerComponentTranslation = ComputeDistributeTranslations(false, EUVEditorAlignDirectionBackend::Left, 1.0f, false);
+			break;
+		case EUVEditorDistributeModeBackend::RightEdges:
+			PerComponentTranslation = ComputeDistributeTranslations(false, EUVEditorAlignDirectionBackend::Right, -1.0f, false);
+			break;
+		case EUVEditorDistributeModeBackend::CentersVertically:
+			PerComponentTranslation = ComputeDistributeTranslations(false, EUVEditorAlignDirectionBackend::CenterVertically, 1.0f, false);
+			break;
+		case EUVEditorDistributeModeBackend::CentersHorizontally:
+			PerComponentTranslation = ComputeDistributeTranslations(true, EUVEditorAlignDirectionBackend::CenterHorizontally, 1.0f, false);
+			break;
+		case EUVEditorDistributeModeBackend::HorizontalSpace:
+			PerComponentTranslation = ComputeDistributeTranslations(false, EUVEditorAlignDirectionBackend::Left, 1.0f, true);
+			break;
+		case EUVEditorDistributeModeBackend::VerticalSpace:
+			PerComponentTranslation = ComputeDistributeTranslations(true, EUVEditorAlignDirectionBackend::Bottom, 1.0f, true);
+			break;
+		case EUVEditorDistributeModeBackend::MinimallyRemoveOverlap:
+			PerComponentTranslation = ComputeMinimallyRemoveOverlap();
+			break;
+		default:
+			ensure(false);
 	}
 
 	auto TranslateFunc = [this](const FVector2f& Translation, int32 Vid)
 	{
-		FVector2f UV = FUVEditorUXSettings::InternalUVToExternalUV(ActiveUVLayer->GetElement(Vid));
+		FVector2f UV = TransformOpLocals::UnwrapPositionToUV(ResultMesh->GetVertexRef(Vid));
 		UV = (UV + Translation);
-		ActiveUVLayer->SetElement(Vid, FUVEditorUXSettings::ExternalUVToInternalUV(UV));
+		ResultMesh->SetVertex(Vid, TransformOpLocals::UVToUnwrapPosition(UV));
 	};
 
 	if (TransformingElements.IsSet())
@@ -532,24 +712,29 @@ void FUVEditorUVDistributeOp::HandleTransformationOp(FProgressCancel* Progress)
 
 TUniquePtr<FDynamicMeshOperator> UUVEditorUVTransformOperatorFactory::MakeNewOperator()
 {
-	switch (Settings->TransformType)
+	switch (TransformType)
 	{
 	case EUVEditorUVTransformType::Transform:
 	{
+		UUVEditorUVTransformProperties* OpSettings = CastChecked<UUVEditorUVTransformProperties>(Settings.Get());
+
 		TUniquePtr<FUVEditorUVTransformOp> Op = MakeUnique<FUVEditorUVTransformOp>();
 		Op->OriginalMesh = OriginalMesh;
 		Op->SetTransform(TargetTransform);
-		Op->Selection = Selection;
+		if (TriangleSelection.IsSet() && VertexSelection.IsSet())
+		{
+			Op->SetSelection(TriangleSelection.GetValue(), VertexSelection.GetValue());
+		}
 		Op->UVLayerIndex = GetSelectedUVChannel();
 
-		Op->Scale = Settings->Scale;
-		Op->Rotation = Settings->Rotation;
-		Op->Translation = Settings->Translation;
+		Op->Scale = OpSettings->Scale;
+		Op->Rotation = OpSettings->Rotation;
+		Op->Translation = OpSettings->Translation;
 
-		Op->QuickTranslation = Settings->QuickTranslation;
-		Op->QuickRotation = Settings->QuickRotation;
+		Op->QuickTranslation = OpSettings->QuickTranslation;
+		Op->QuickRotation = OpSettings->QuickRotation;
 
-		switch (Settings->TranslationMode)
+		switch (OpSettings->TranslationMode)
 		{
 		case EUVEditorTranslationMode::Relative:
 			Op->TranslationMode = EUVEditorTranslationModeBackend::Relative;
@@ -557,9 +742,11 @@ TUniquePtr<FDynamicMeshOperator> UUVEditorUVTransformOperatorFactory::MakeNewOpe
 		case EUVEditorTranslationMode::Absolute:
 			Op->TranslationMode = EUVEditorTranslationModeBackend::Absolute;
 			break;
+		default:
+			ensure(false);
 		}
 
-		switch (Settings->PivotMode)
+		switch (OpSettings->PivotMode)
 		{
 		case EUVEditorPivotType::Origin:
 			Op->PivotMode = EUVEditorPivotTypeBackend::Origin;
@@ -573,22 +760,29 @@ TUniquePtr<FDynamicMeshOperator> UUVEditorUVTransformOperatorFactory::MakeNewOpe
 		case EUVEditorPivotType::Manual:
 			Op->PivotMode = EUVEditorPivotTypeBackend::Manual;
 			break;
+		default:
+			ensure(false);
 		}
-		Op->ManualPivot = Settings->ManualPivot;
-
+		Op->ManualPivot = OpSettings->ManualPivot;
+		Op->GroupingMode = EUVEditorAlignDistributeGroupingModeBackend::IndividualBoundingBoxes;
 
 		return Op;
 	}
 	break;
 	case EUVEditorUVTransformType::Align:
 	{
+		UUVEditorUVAlignProperties* OpSettings = CastChecked<UUVEditorUVAlignProperties>(Settings.Get());
+
 		TUniquePtr<FUVEditorUVAlignOp> Op = MakeUnique<FUVEditorUVAlignOp>();
 		Op->OriginalMesh = OriginalMesh;
 		Op->SetTransform(TargetTransform);
-		Op->Selection = Selection;
+		if (TriangleSelection.IsSet() && VertexSelection.IsSet())
+		{
+			Op->SetSelection(TriangleSelection.GetValue(), VertexSelection.GetValue());
+		}
 		Op->UVLayerIndex = GetSelectedUVChannel();
 
-		switch (Settings->AlignAnchor)
+		switch (OpSettings->AlignAnchor)
 		{
 		//case EUVEditorAlignAnchor::FirstItem:
 		//	Op->AlignAnchor = EUVEditorAlignAnchorBackend::FirstItem;
@@ -602,10 +796,15 @@ TUniquePtr<FDynamicMeshOperator> UUVEditorUVTransformOperatorFactory::MakeNewOpe
 		case EUVEditorAlignAnchor::Manual:
 			Op->AlignAnchor = EUVEditorAlignAnchorBackend::Manual;
 			break;
+		default:
+			ensure(false);
 		}
 
-		switch (Settings->AlignDirection)
+		switch (OpSettings->AlignDirection)
 		{
+		case EUVEditorAlignDirection::None:
+			Op->AlignDirection = EUVEditorAlignDirectionBackend::None;
+			break;
 		case EUVEditorAlignDirection::Top:
 			Op->AlignDirection = EUVEditorAlignDirectionBackend::Top;
 			break;
@@ -624,23 +823,48 @@ TUniquePtr<FDynamicMeshOperator> UUVEditorUVTransformOperatorFactory::MakeNewOpe
 		case EUVEditorAlignDirection::CenterHorizontally:
 			Op->AlignDirection = EUVEditorAlignDirectionBackend::CenterHorizontally;
 			break;
+		default:
+			ensure(false);
+		}
+
+		switch (OpSettings->Grouping)
+		{
+		case EUVEditorAlignDistributeGroupingMode::IndividualBoundingBoxes:
+			Op->GroupingMode = EUVEditorAlignDistributeGroupingModeBackend::IndividualBoundingBoxes;
+			break;
+		case EUVEditorAlignDistributeGroupingMode::EnclosingBoundingBox:
+			Op->GroupingMode = EUVEditorAlignDistributeGroupingModeBackend::EnclosingBoundingBox;
+			break;
+		case EUVEditorAlignDistributeGroupingMode::IndividualVertices:
+			Op->GroupingMode = EUVEditorAlignDistributeGroupingModeBackend::IndividualVertices;
+			break;
+		default:
+			ensure(false);
 		}
 
-		Op->ManualAnchor = Settings->ManualAnchor;
+		Op->ManualAnchor = OpSettings->ManualAnchor;
 
 		return Op;
 	}
 	break;
 	case EUVEditorUVTransformType::Distribute:
 	{
+		UUVEditorUVDistributeProperties* OpSettings = CastChecked<UUVEditorUVDistributeProperties>(Settings.Get());
+
 		TUniquePtr<FUVEditorUVDistributeOp> Op = MakeUnique<FUVEditorUVDistributeOp>();
 		Op->OriginalMesh = OriginalMesh;
 		Op->SetTransform(TargetTransform);
-		Op->Selection = Selection;
+		if (TriangleSelection.IsSet() && VertexSelection.IsSet())
+		{
+			Op->SetSelection(TriangleSelection.GetValue(), VertexSelection.GetValue());
+		}
 		Op->UVLayerIndex = GetSelectedUVChannel();
 
-		switch (Settings->DistributeMode)
+		switch (OpSettings->DistributeMode)
 		{
+		case EUVEditorDistributeMode::None:
+			Op->DistributeMode = EUVEditorDistributeModeBackend::None;
+			break;
 		case EUVEditorDistributeMode::LeftEdges:
 			Op->DistributeMode = EUVEditorDistributeModeBackend::LeftEdges;
 			break;
@@ -665,8 +889,32 @@ TUniquePtr<FDynamicMeshOperator> UUVEditorUVTransformOperatorFactory::MakeNewOpe
 		case EUVEditorDistributeMode::HorizontalSpace:
 			Op->DistributeMode = EUVEditorDistributeModeBackend::HorizontalSpace;
 			break;
+		case EUVEditorDistributeMode::MinimallyRemoveOverlap:
+			Op->DistributeMode = EUVEditorDistributeModeBackend::MinimallyRemoveOverlap;
+			break;
+		default:
+			ensure(false);
+		}
+
+		switch (OpSettings->Grouping)
+		{
+		case EUVEditorAlignDistributeGroupingMode::IndividualBoundingBoxes:
+			Op->GroupingMode = EUVEditorAlignDistributeGroupingModeBackend::IndividualBoundingBoxes;
+			break;
+		case EUVEditorAlignDistributeGroupingMode::EnclosingBoundingBox:
+			Op->GroupingMode = EUVEditorAlignDistributeGroupingModeBackend::EnclosingBoundingBox;
+			break;
+		case EUVEditorAlignDistributeGroupingMode::IndividualVertices:
+			Op->GroupingMode = EUVEditorAlignDistributeGroupingModeBackend::IndividualVertices;
+			break;
+		default:
+			ensure(false);
 		}
 
+		Op->bEnableManualDistances = OpSettings->bEnableManualDistances;
+		Op->ManualExtent = OpSettings->ManualExtent;
+		Op->ManualSpacing = OpSettings->ManualSpacing;
+
 		return Op;
 	}
 	break;
diff --git a/Engine/Plugins/Experimental/UVEditor/Source/UVEditorTools/Private/UVEditorStyle.cpp b/Engine/Plugins/Experimental/UVEditor/Source/UVEditorTools/Private/UVEditorStyle.cpp
new file mode 100644
--- /dev/null
+++ b/Engine/Plugins/Experimental/UVEditor/Source/UVEditorTools/Private/UVEditorStyle.cpp
@@ -0,0 +1,88 @@
+// Copyright Epic Games, Inc. All Rights Reserved.
+
+#include "UVEditorStyle.h"
+
+#include "Styling/SlateTypes.h"
+#include "Styling/CoreStyle.h"
+#include "Styling/AppStyle.h"
+#include "Styling/SlateStyleMacros.h"
+
+FName FUVEditorStyle::StyleName("UVStyle");
+
+FUVEditorStyle::FUVEditorStyle()
+	: FSlateStyleSet(StyleName)
+{
+	// Used FFractureEditorStyle as a model
+
+	const FVector2D IconSize(16.0f, 16.0f);
+	const FVector2D ToolbarIconSize(20.0f, 20.0f);
+
+	SetContentRoot(FPaths::EnginePluginsDir() / TEXT("Experimental/UVEditor/Content/Icons"));
+	SetCoreContentRoot(FPaths::EngineContentDir() / TEXT("Editor/Slate") );
+
+	Set("UVEditor.OpenUVEditor",				new IMAGE_BRUSH_SVG("UVEditor", IconSize));
+
+	Set("UVEditor.BeginSelectTool",				new FSlateVectorImageBrush(FPaths::EngineContentDir() / TEXT("Slate/Starship/Common/edit.svg"), ToolbarIconSize));
+	Set("UVEditor.BeginLayoutTool",             new IMAGE_BRUSH_SVG("UVLayout", ToolbarIconSize));
+	Set("UVEditor.BeginParameterizeMeshTool",   new IMAGE_BRUSH_SVG("AutoUnwrap", ToolbarIconSize));
+	Set("UVEditor.BeginChannelEditTool",        new IMAGE_BRUSH_SVG("AttributeEditor", ToolbarIconSize));
+	Set("UVEditor.BeginSeamTool",               new IMAGE_BRUSH_SVG("ModelingUVSeamEdit", ToolbarIconSize));
+	Set("UVEditor.BeginRecomputeUVsTool",       new IMAGE_BRUSH_SVG("GroupUnwrap", ToolbarIconSize));
+	Set("UVEditor.BeginTransformTool",          new IMAGE_BRUSH_SVG("TransformUVs", ToolbarIconSize));
+	Set("UVEditor.BeginAlignTool",              new IMAGE_BRUSH_SVG("AlignLeft", ToolbarIconSize));
+	Set("UVEditor.BeginDistributeTool",         new IMAGE_BRUSH_SVG("DistributeHorizontally", ToolbarIconSize));
+
+	// Select tool actions
+	Set("UVEditor.SewAction", new IMAGE_BRUSH_SVG("UVSew", ToolbarIconSize));
+	Set("UVEditor.SplitAction", new IMAGE_BRUSH_SVG("UVCut", ToolbarIconSize));
+	Set("UVEditor.IslandConformalUnwrapAction", new IMAGE_BRUSH_SVG("UVUnwrap", ToolbarIconSize));
+
+	// Top toolbar icons
+	Set("UVEditor.ApplyChanges", new CORE_IMAGE_BRUSH_SVG("Starship/Common/Apply", ToolbarIconSize));
+	Set("UVEditor.ChannelSettings", new CORE_IMAGE_BRUSH_SVG("Starship/Common/SetDrawUVs", ToolbarIconSize));
+	Set("UVEditor.BackgroundSettings", new CORE_IMAGE_BRUSH_SVG("Starship/Common/Sprite", ToolbarIconSize));
+
+	// Viewport icons
+	Set("UVEditor.OrbitCamera", new CORE_IMAGE_BRUSH_SVG("Starship/EditorViewport/rotate", ToolbarIconSize));
+	Set("UVEditor.FlyCamera", new CORE_IMAGE_BRUSH_SVG("Starship/EditorViewport/camera", ToolbarIconSize));
+
+	Set("UVEditor.VertexSelection", new IMAGE_BRUSH_SVG("SelectionVertices", ToolbarIconSize));
+	Set("UVEditor.EdgeSelection", new IMAGE_BRUSH_SVG("SelectionLine", ToolbarIconSize));
+	Set("UVEditor.TriangleSelection", new IMAGE_BRUSH_SVG("SelectionTriangle", ToolbarIconSize));
+	Set("UVEditor.IslandSelection", new IMAGE_BRUSH_SVG("SelectionIslands", ToolbarIconSize));
+	Set("UVEditor.FullMeshSelection", new IMAGE_BRUSH_SVG("SelectionMulti", ToolbarIconSize));
+
+	// Distribute
+	Set("UVEditor.DistributeSpaceVertically", new IMAGE_BRUSH_SVG("DistributeModeSpaceVertically", ToolbarIconSize));
+	Set("UVEditor.DistributeSpaceHorizontally", new IMAGE_BRUSH_SVG("DistributeModeSpaceHorizontally", ToolbarIconSize));
+	Set("UVEditor.DistributeTopEdges", new IMAGE_BRUSH_SVG("DistributeModeTopEdges", ToolbarIconSize));
+	Set("UVEditor.DistributeBottomEdges", new IMAGE_BRUSH_SVG("DistributeModeBottomEdges", ToolbarIconSize));
+	Set("UVEditor.DistributeLeftEdges", new IMAGE_BRUSH_SVG("DistributeModeLeftEdges", ToolbarIconSize));
+	Set("UVEditor.DistributeRightEdges", new IMAGE_BRUSH_SVG("DistributeModeRightEdges", ToolbarIconSize));
+	Set("UVEditor.DistributeCentersHorizontally", new IMAGE_BRUSH_SVG("DistributeModeCentersHorizontally", ToolbarIconSize));
+	Set("UVEditor.DistributeCentersVertically", new IMAGE_BRUSH_SVG("DistributeModeCentersVertically", ToolbarIconSize));
+	Set("UVEditor.DistributeRemoveOverlap", new IMAGE_BRUSH_SVG("DistributeModeRemoveOverlap", ToolbarIconSize));
+
+	// Align
+	Set("UVEditor.AlignDirectionBottomEdges", new IMAGE_BRUSH_SVG("DistributeModeBottomEdges", ToolbarIconSize));
+	Set("UVEditor.AlignDirectionTopEdges", new IMAGE_BRUSH_SVG("DistributeModeTopEdges", ToolbarIconSize));
+	Set("UVEditor.AlignDirectionLeftEdges", new IMAGE_BRUSH_SVG("DistributeModeLeftEdges", ToolbarIconSize));
+	Set("UVEditor.AlignDirectionRightEdges", new IMAGE_BRUSH_SVG("DistributeModeRightEdges", ToolbarIconSize));
+	Set("UVEditor.AlignDirectionCentersHorizontally", new IMAGE_BRUSH_SVG("DistributeModeCentersVertically", ToolbarIconSize));
+	Set("UVEditor.AlignDirectionCentersVertically", new IMAGE_BRUSH_SVG("DistributeModeCentersHorizontally", ToolbarIconSize));
+
+	FSlateStyleRegistry::RegisterSlateStyle(*this);
+}
+
+FUVEditorStyle::~FUVEditorStyle()
+{
+	FSlateStyleRegistry::UnRegisterSlateStyle(*this);
+}
+
+FUVEditorStyle& FUVEditorStyle::Get()
+{
+	static FUVEditorStyle Inst;
+	return Inst;
+}
+
+
diff --git a/Engine/Plugins/Experimental/UVEditor/Source/UVEditorTools/Private/UVEditorTransformTool.cpp b/Engine/Plugins/Experimental/UVEditor/Source/UVEditorTools/Private/UVEditorTransformTool.cpp
--- a/Engine/Plugins/Experimental/UVEditor/Source/UVEditorTools/Private/UVEditorTransformTool.cpp
+++ b/Engine/Plugins/Experimental/UVEditor/Source/UVEditorTools/Private/UVEditorTransformTool.cpp
@@ -71,50 +71,68 @@ void UUVEditorTransformTool::Setup()
 
 	UInteractiveTool::Setup();
 
-	Settings = NewObject<UUVEditorUVTransformProperties>(this);
+	switch(ToolMode.Get(EUVEditorUVTransformType::Transform))
+	{
+		case EUVEditorUVTransformType::Transform:
+			Settings = NewObject<UUVEditorUVTransformProperties>(this);
+			break;
+		case EUVEditorUVTransformType::Align:
+			Settings = NewObject<UUVEditorUVAlignProperties>(this);
+			break;
+		case EUVEditorUVTransformType::Distribute:
+			Settings = NewObject<UUVEditorUVDistributeProperties>(this);
+			break;
+		default:
+			ensure(false);
+	}
 	Settings->RestoreProperties(this);
-	Settings->TransformType = ToolMode.Get(EUVEditorUVTransformType::Transform);
 	//Settings->bUDIMCVAREnabled = (FUVEditorUXSettings::CVarEnablePrototypeUDIMSupport.GetValueOnGameThread() > 0);
 	AddToolPropertySource(Settings);
 
-	if (ToolMode.Get(EUVEditorUVTransformType::Transform) == EUVEditorUVTransformType::Transform)
-	{
-		QuickTransformSettings = NewObject<UUVEditorUVQuickTransformProperties>(this);
-		QuickTransformSettings->RestoreProperties(this);
-		QuickTransformSettings->Tool = this;
-		AddToolPropertySource(QuickTransformSettings);
-	}
-
 	UContextObjectStore* ContextStore = GetToolManager()->GetContextObjectStore();
 	UVToolSelectionAPI = ContextStore->FindContext<UUVToolSelectionAPI>();
 
 	UUVToolSelectionAPI::FHighlightOptions HighlightOptions;
 	HighlightOptions.bBaseHighlightOnPreviews = true;
 	HighlightOptions.bAutoUpdateUnwrap = true;
-	UVToolSelectionAPI->SetSelectionMechanicMode(UUVToolSelectionAPI::EUVEditorSelectionMode::Triangle);
 	UVToolSelectionAPI->SetHighlightOptions(HighlightOptions);
 	UVToolSelectionAPI->SetHighlightVisible(true, false, true);
-
-	auto SetupOpFactory = [this](UUVEditorToolMeshInput& Target, TSet<int32>* Selection)
+	
+	auto SetupOpFactory = [this](UUVEditorToolMeshInput& Target, const FUVToolSelection* Selection)
 	{
 		TObjectPtr<UUVEditorUVTransformOperatorFactory> Factory = NewObject<UUVEditorUVTransformOperatorFactory>();
 		Factory->TargetTransform = Target.AppliedPreview->PreviewMesh->GetTransform();
 		Factory->Settings = Settings;
-		Factory->OriginalMesh = Target.AppliedCanonical;
+		Factory->TransformType = ToolMode.Get(EUVEditorUVTransformType::Transform);
+		Factory->OriginalMesh = Target.UnwrapCanonical;
 		Factory->GetSelectedUVChannel = [&Target]() { return Target.UVLayerIndex; };
 		if (Selection)
 		{
-			Factory->Selection.Emplace(*Selection);
+			// If we have a selection, it's in the unwrapped mesh. We need both triangles, which are 1:1 between them,
+			// and vertices, which are not, in the applied mesh to pass to the factory.
+			FUVToolSelection UnwrapVertexSelection;
+			if (Selection->Type == FUVToolSelection::EType::Vertex)
+			{
+				UnwrapVertexSelection = *Selection;
+			}
+			else
+			{
+				UnwrapVertexSelection = Selection->GetConvertedSelection(*Target.UnwrapCanonical, FUVToolSelection::EType::Vertex);
+			}				
+			Factory->VertexSelection.Emplace(UnwrapVertexSelection.SelectedIDs);
+			Factory->TriangleSelection.Emplace(UnwrapVertexSelection.GetConvertedSelection(*Target.UnwrapCanonical, FUVToolSelection::EType::Triangle).SelectedIDs);
 		}
 
-		Target.AppliedPreview->ChangeOpFactory(Factory);
-		Target.AppliedPreview->OnMeshUpdated.AddWeakLambda(this, [this, &Target](UMeshOpPreviewWithBackgroundCompute* Preview) {
-			Target.UpdateUnwrapPreviewFromAppliedPreview();
+		Target.UnwrapPreview->ChangeOpFactory(Factory);
+		Target.UnwrapPreview->OnMeshUpdated.AddWeakLambda(this, [this, &Target](UMeshOpPreviewWithBackgroundCompute* Preview) {
+			Target.UpdateUnwrapPreviewOverlayFromPositions();
+			Target.UpdateAppliedPreviewFromUnwrapPreview();
+			
 
 			this->UVToolSelectionAPI->RebuildUnwrapHighlight(Preview->PreviewMesh->GetTransform());
 			});
 
-		Target.AppliedPreview->InvalidateResult();
+		Target.UnwrapPreview->InvalidateResult();
 		return Factory;
 	};
 
@@ -123,10 +141,7 @@ void UUVEditorTransformTool::Setup()
 		Factories.Reserve(UVToolSelectionAPI->GetSelections().Num());
 		for (FUVToolSelection Selection : UVToolSelectionAPI->GetSelections())
 		{
-			if (ensure(Selection.Type == FUVToolSelection::EType::Triangle))
-			{
-				Factories.Add(SetupOpFactory(*Selection.Target, &Selection.SelectedIDs));
-			}
+			Factories.Add(SetupOpFactory(*Selection.Target, &Selection));
 		}
 	}
 	else
@@ -151,7 +166,7 @@ void UUVEditorTransformTool::Shutdown(EToolShutdownType ShutdownType)
 	Settings->SaveProperties(this);
 	for (TObjectPtr<UUVEditorToolMeshInput> Target : Targets)
 	{
-		Target->AppliedPreview->OnMeshUpdated.RemoveAll(this);
+		Target->UnwrapPreview->OnMeshUpdated.RemoveAll(this);
 	}
 
 	if (ShutdownType == EToolShutdownType::Accept)
@@ -197,7 +212,7 @@ void UUVEditorTransformTool::Shutdown(EToolShutdownType ShutdownType)
 
 	for (TObjectPtr<UUVEditorToolMeshInput> Target : Targets)
 	{
-		Target->AppliedPreview->ClearOpFactory();
+		Target->UnwrapPreview->ClearOpFactory();
 	}
 
 	for (int32 FactoryIndex = 0; FactoryIndex < Factories.Num(); ++FactoryIndex)
@@ -213,26 +228,7 @@ void UUVEditorTransformTool::OnTick(float DeltaTime)
 {
 	for (TObjectPtr<UUVEditorToolMeshInput> Target : Targets)
 	{
-		Target->AppliedPreview->Tick(DeltaTime);
-	}
-}
-
-void UUVEditorTransformTool::InitiateQuickTranslate(float Offset, const FVector2D& Direction)
-{
-	Settings->QuickTranslation += Direction * Offset;
-	for (TObjectPtr<UUVEditorToolMeshInput> Target : Targets)
-	{
-		Target->AppliedPreview->InvalidateResult();
-	}
-}
-
-
-void UUVEditorTransformTool::InitiateQuickRotation(float Rotation)
-{
-	Settings->QuickRotation += Rotation;
-	for (TObjectPtr<UUVEditorToolMeshInput> Target : Targets)
-	{
-		Target->AppliedPreview->InvalidateResult();
+		Target->UnwrapPreview->Tick(DeltaTime);
 	}
 }
 
@@ -240,7 +236,7 @@ void UUVEditorTransformTool::OnPropertyModified(UObject* PropertySet, FProperty*
 {
 	for (TObjectPtr<UUVEditorToolMeshInput> Target : Targets)
 	{
-		Target->AppliedPreview->InvalidateResult();
+		Target->UnwrapPreview->InvalidateResult();
 	}
 }
 
@@ -248,7 +244,7 @@ bool UUVEditorTransformTool::CanAccept() const
 {
 	for (TObjectPtr<UUVEditorToolMeshInput> Target : Targets)
 	{
-		if (!Target->AppliedPreview->HaveValidResult())
+		if (!Target->UnwrapPreview->HaveValidResult())
 		{
 			return false;
 		}
@@ -284,7 +280,7 @@ void UUVEditorTransformTool::RecordAnalytics()
 		for (TObjectPtr<UUVEditorToolMeshInput> Target : Targets)
 		{
 			// Note: This would log -1 if the result was invalid, but checking CanAccept above ensures results are valid
-			PerAssetValidResultComputeTimes.Add(Target->AppliedPreview->GetValidResultComputeTime());
+			PerAssetValidResultComputeTimes.Add(Target->UnwrapPreview->GetValidResultComputeTime());
 		}
 		Attributes.Add(FAnalyticsEventAttribute(TEXT("Stats.PerAsset.ComputeTimeSeconds"), PerAssetValidResultComputeTimes));
 	}
diff --git a/Engine/Plugins/Experimental/UVEditor/Source/UVEditorTools/Public/Operators/UVEditorUVTransformOp.h b/Engine/Plugins/Experimental/UVEditor/Source/UVEditorTools/Public/Operators/UVEditorUVTransformOp.h
--- a/Engine/Plugins/Experimental/UVEditor/Source/UVEditorTools/Public/Operators/UVEditorUVTransformOp.h
+++ b/Engine/Plugins/Experimental/UVEditor/Source/UVEditorTools/Public/Operators/UVEditorUVTransformOp.h
@@ -7,6 +7,7 @@
 #include "ModelingOperators.h"
 #include "InteractiveTool.h"
 #include "DynamicMesh/DynamicMeshAttributeSet.h"
+#include "Selection/UVToolSelection.h"
 
 #include "UVEditorUVTransformOp.generated.h"
 
@@ -58,6 +59,7 @@ enum class EUVEditorTranslationMode
 UENUM()
 enum class EUVEditorAlignDirection
 {
+	None,
 	Top,
 	Bottom,
 	Left,
@@ -80,8 +82,9 @@ enum class EUVEditorAlignAnchor
 };
 
 UENUM()
-enum class EUVEditorDistributeMode
+enum class EUVEditorDistributeMode : uint16
 {
+	None,
 	VerticalSpace,
 	HorizontalSpace,
 	TopEdges,
@@ -89,79 +92,155 @@ enum class EUVEditorDistributeMode
 	LeftEdges,
 	RightEdges,
 	CentersVertically,
-	CentersHorizontally
+	CentersHorizontally,
+	MinimallyRemoveOverlap
 };
 
+UENUM()
+enum class EUVEditorAlignDistributeGroupingMode : uint8
+{
+	/** Treat selection as individual connected components, moving each according to their bounding boxes.*/
+	IndividualBoundingBoxes,
+	/** Treat selection as a single unit, moving everything together and preserving relative positions.*/
+	EnclosingBoundingBox,
+	/** Treat selection as isolated UV vertices, moving each UV independently of each other.*/
+	IndividualVertices
+};
 
 /**
  * UV Transform Settings
  */
 UCLASS()
-class UVEDITORTOOLS_API UUVEditorUVTransformProperties : public UInteractiveToolPropertySet
+class UVEDITORTOOLS_API UUVEditorUVTransformPropertiesBase : public UInteractiveToolPropertySet
 {
 	GENERATED_BODY()
-
 public:
-	/** Type of Transform applied to input UVs */
-	UPROPERTY()
-	EUVEditorUVTransformType TransformType = EUVEditorUVTransformType::Transform;
+};
+
 
+/**
+ * UV Transform Settings
+ */
+UCLASS()
+class UVEDITORTOOLS_API UUVEditorUVTransformProperties : public UUVEditorUVTransformPropertiesBase
+{
+	GENERATED_BODY()
+
+public:
 	/** Scale applied to UVs, potentially non-uniform */
-	UPROPERTY(EditAnywhere, Category = "Advanced Transform | Scaling", DuplicateTransient, meta = (TransientToolProperty, DisplayName = "Scale",
-		      EditCondition = "TransformType == EUVEditorUVTransformType::Transform", EditConditionHides, HideEditConditionToggle = true))
+	UPROPERTY(EditAnywhere, Category = "Advanced Transform | Scaling", DuplicateTransient, meta = (TransientToolProperty, DisplayName = "Scale"))
 	FVector2D Scale = FVector2D(1.0, 1.0); 
 
 	/** Rotation applied to UVs after scaling, specified in degrees */
-	UPROPERTY(EditAnywhere, Category = "Advanced Transform | Rotation", meta = (TransientToolProperty, DisplayName = "Rotation",
-		      EditCondition = "TransformType == EUVEditorUVTransformType::Transform", EditConditionHides, HideEditConditionToggle = true))
+	UPROPERTY(EditAnywhere, Category = "Advanced Transform | Rotation", meta = (TransientToolProperty, DisplayName = "Rotation"))
 	float Rotation = 0.0;
 
 	/** Translation applied to UVs, and after scaling and rotation */
-	UPROPERTY(EditAnywhere, Category = "Advanced Transform | Translation", meta = (TransientToolProperty, DisplayName = "Translation",
-		      EditCondition = "TransformType == EUVEditorUVTransformType::Transform", EditConditionHides, HideEditConditionToggle = true))
+	UPROPERTY(EditAnywhere, Category = "Advanced Transform | Translation", meta = (TransientToolProperty, DisplayName = "Translation"))
 	FVector2D Translation = FVector2D(0, 0);
 
 	/** Translation applied to UVs, and after scaling and rotation */
-	UPROPERTY(EditAnywhere, Category = "Advanced Transform | Translation", meta = (TransientToolProperty, DisplayName = "Translation Mode",
-		EditCondition = "TransformType == EUVEditorUVTransformType::Transform", EditConditionHides, HideEditConditionToggle = true))
+	UPROPERTY(EditAnywhere, Category = "Advanced Transform | Translation", meta = (TransientToolProperty, DisplayName = "Translation Mode"))
 	EUVEditorTranslationMode TranslationMode = EUVEditorTranslationMode::Relative;
 
 	/** Transformation origin mode used for scaling and rotation */
-	UPROPERTY(EditAnywhere, Category = "Advanced Transform | Transform Origin", meta = (TransientToolProperty, DisplayName = "Mode",
-		EditCondition = "TransformType == EUVEditorUVTransformType::Transform", EditConditionHides, HideEditConditionToggle = true))
+	UPROPERTY(EditAnywhere, Category = "Advanced Transform | Transform Origin", meta = (TransientToolProperty, DisplayName = "Mode"))
 	EUVEditorPivotType PivotMode = EUVEditorPivotType::BoundingBoxCenter;
 
 	/** Manual Transformation origin point */
 	UPROPERTY(EditAnywhere, Category = "Advanced Transform | Transform Origin", meta = (TransientToolProperty, DisplayName = "Coords",
-		EditCondition = "TransformType == EUVEditorUVTransformType::Transform && PivotMode == EUVEditorPivotType::Manual", EditConditionHides, HideEditConditionToggle = true))
+		EditCondition = "PivotMode == EUVEditorPivotType::Manual", EditConditionHides, HideEditConditionToggle = true))
 	FVector2D ManualPivot = FVector2D(0, 0);
 
-	UPROPERTY(meta = (TransientToolProperty))
+	UPROPERTY(EditAnywhere, Category = "Quick Transform", meta = (TransientToolProperty, DisplayName = "Translation"))
+	float QuickTranslateOffset = 0.0;
+
+	UPROPERTY(EditAnywhere, Category = "Quick Transform", meta = (TransientToolProperty, DisplayName = "Rotation"))
+	float QuickRotationOffset = 0.0;
+
+	UPROPERTY(EditAnywhere, Category = "Quick Transform", meta = (TransientToolProperty))
 	FVector2D QuickTranslation;
 
-	UPROPERTY(meta = (TransientToolProperty))
+	UPROPERTY(EditAnywhere, Category = "Quick Transform", meta = (TransientToolProperty))
 	float QuickRotation;
+};
+
+/**
+ * UV Quick Transform Only Settings
+ * 
+ * We are using a subclass here to "trick" the details customization system, allowing us to reuse a lot of the logic
+ * without having to build new operators or new customizations.
+ * 
+ * See FUVEditorUVTransformToolDetails and FUVEditorUVQuickTransformToolDetails, where we provide two customizations
+ * which present different views of the "same" properties - one with all settings and one with simply the quick transform settings.
+ * This is designed to be used in the future where we want to provide a quick translate "tool" when no other tools are running.
+ * 
+ */
+UCLASS()
+class UVEDITORTOOLS_API UUVEditorUVQuickTransformProperties : public UUVEditorUVTransformProperties
+{
+	GENERATED_BODY()
+};
 
+/**
+ * UV Align Settings
+ */
+UCLASS()
+class UVEDITORTOOLS_API UUVEditorUVAlignProperties : public UUVEditorUVTransformPropertiesBase
+{
+	GENERATED_BODY()
+
+public:	
 
 	/** Controls what geometry the alignment is to be relative to when performed. */
-	UPROPERTY(EditAnywhere, Category = "Align", meta = (DisplayName = "Alignment Anchor",
-		EditCondition = "TransformType == EUVEditorUVTransformType::Align", EditConditionHides, HideEditConditionToggle = true))
+	UPROPERTY(EditAnywhere, Category = "Align", meta = (DisplayName = "Alignment Anchor"))
 	EUVEditorAlignAnchor AlignAnchor = EUVEditorAlignAnchor::BoundingBox;
 
 	/** Manual anchor location for relative alignment */
 	UPROPERTY(EditAnywhere, Category = "Align", meta = (DisplayName = "Anchor Coords",
-		EditCondition = "TransformType == EUVEditorUVTransformType::Align && AlignAnchor == EUVEditorAlignAnchor::Manual", EditConditionHides, HideEditConditionToggle = true))
+		EditCondition = "AlignAnchor == EUVEditorAlignAnchor::Manual", EditConditionHides, HideEditConditionToggle = true))
 	FVector2D ManualAnchor = FVector2D(0, 0);
 
 	/** Controls what side of the island bounding boxes are being aligned */
-	UPROPERTY(EditAnywhere, Category = "Align", meta = (DisplayName = "Alignment Direction",
-		EditCondition = "TransformType == EUVEditorUVTransformType::Align", EditConditionHides, HideEditConditionToggle = true))
-	EUVEditorAlignDirection AlignDirection = EUVEditorAlignDirection::Top;
+	UPROPERTY(EditAnywhere, Category = "Align", meta = (DisplayName = "Alignment Direction"))
+	EUVEditorAlignDirection AlignDirection = EUVEditorAlignDirection::None;
+
+	/** Controls how alignment considers grouping selected objects with respect to the alignment behavior.*/
+	UPROPERTY(EditAnywhere, Category = "Align", meta = (DisplayName = "Grouping Mode"))
+	EUVEditorAlignDistributeGroupingMode Grouping = EUVEditorAlignDistributeGroupingMode::IndividualBoundingBoxes;
 
+};
+
+/**
+ * UV Distribute Settings
+ */
+UCLASS()
+class UVEDITORTOOLS_API UUVEditorUVDistributeProperties : public UUVEditorUVTransformPropertiesBase
+{
+	GENERATED_BODY()
+public:
 	/** Controls the distribution behavior */
-	UPROPERTY(EditAnywhere, Category = "Distribute", meta = (DisplayName = "Distribution Mode",
-		EditCondition = "TransformType == EUVEditorUVTransformType::Distribute", EditConditionHides, HideEditConditionToggle = true))
-	EUVEditorDistributeMode DistributeMode = EUVEditorDistributeMode::TopEdges;
+	UPROPERTY(EditAnywhere, Category = "Distribute", meta = (TransientToolProperty, DisplayName = "Distribution Mode"))
+	EUVEditorDistributeMode DistributeMode = EUVEditorDistributeMode::None;
+
+	/** Controls how distribution considers grouping selected objects with respect to the distribution behavior.*/
+	UPROPERTY(EditAnywhere, Category = "Distribute", meta = (DisplayName = "Grouping Mode"))
+	EUVEditorAlignDistributeGroupingMode Grouping = EUVEditorAlignDistributeGroupingMode::IndividualBoundingBoxes;
+
+	/** If true, enable overriding distances used in the distribution behavior with manually entered values.*/
+	UPROPERTY(EditAnywhere, Category = "Distribute", meta = (DisplayName = "Manual Distances"))
+	bool bEnableManualDistances;
+
+	/** For Edge and Center distribution modes, specify the desired overall distance within which to evenly place the edges or centers. */
+	UPROPERTY(EditAnywhere, Category = "Distribute", meta = (DisplayName = "Manual Extent",
+		EditCondition = "bEnableManualDistances", EditConditionHides, HideEditConditionToggle = true))
+	float ManualExtent;
+
+	/** For Spacing and Remove Overlap distribution modes, specify the desired distance between objects. */
+	UPROPERTY(EditAnywhere, Category = "Distribute", meta = (DisplayName = "Manual Spacing",
+		EditCondition = "bEnableManualDistances", EditConditionHides, HideEditConditionToggle = true))
+	float ManualSpacing;
+
 };
 
 
@@ -197,6 +276,7 @@ namespace Geometry
 
 	enum class EUVEditorAlignDirectionBackend
 	{
+		None,
 		Top,
 		Bottom,
 		Left,
@@ -215,6 +295,7 @@ namespace Geometry
 
 	enum class EUVEditorDistributeModeBackend
 	{
+		None,
 		LeftEdges,
 		RightEdges,
 		TopEdges,
@@ -222,7 +303,15 @@ namespace Geometry
 		CentersVertically,
 		CentersHorizontally,
 		VerticalSpace,
-		HorizontalSpace
+		HorizontalSpace,
+		MinimallyRemoveOverlap
+	};
+
+	enum class EUVEditorAlignDistributeGroupingModeBackend
+	{
+		IndividualBoundingBoxes,
+		EnclosingBoundingBox,
+		IndividualVertices
 	};
 
 class UVEDITORTOOLS_API FUVEditorUVTransformBaseOp : public FDynamicMeshOperator
@@ -232,10 +321,11 @@ public:
 
 	// inputs
 	TSharedPtr<FDynamicMesh3, ESPMode::ThreadSafe> OriginalMesh;
-	TOptional<TSet<int32>> Selection;
 	int UVLayerIndex = 0;
+	EUVEditorAlignDistributeGroupingModeBackend GroupingMode;
 
 	void SetTransform(const FTransformSRT3d& Transform);
+	void SetSelection(const TSet<int32>& TriangleSelection, const TSet<int32>& VertexSelection);
 
 	virtual void CalculateResult(FProgressCancel* Progress) override;
 
@@ -246,6 +336,7 @@ protected:
 	FVector2f GetAlignmentPointFromUDIMAndDirection(EUVEditorAlignDirectionBackend Direction, FVector2i UDIMTile);
 	void RebuildBoundingBoxes();
 	void CollectTransformElements();
+	void SortComponentsByBoundingBox();
 
 	TMap<int32, int32> ElementToComponent;
 	FDynamicMeshUVOverlay* ActiveUVLayer;
@@ -253,6 +344,10 @@ protected:
 	TArray<FAxisAlignedBox2d> PerComponentBoundingBoxes;
 	TOptional< TSet<int32> > TransformingElements;
 	TSharedPtr<FMeshConnectedComponents> UVComponents;
+	TArray<int32> SpatiallyOrderedComponentIndex;
+
+	TOptional<TSet<int32>> TriangleSelection;
+	TOptional<TSet<int32>> VertexSelection;
 };
 
 
@@ -286,7 +381,7 @@ public:
 	virtual ~FUVEditorUVAlignOp() {}
 
 	EUVEditorAlignAnchorBackend AlignAnchor = EUVEditorAlignAnchorBackend::BoundingBox;
-	EUVEditorAlignDirectionBackend AlignDirection = EUVEditorAlignDirectionBackend::Top;
+	EUVEditorAlignDirectionBackend AlignDirection = EUVEditorAlignDirectionBackend::None;
 	FVector2D ManualAnchor = FVector2D(0, 0);
 
 protected:
@@ -301,7 +396,10 @@ class UVEDITORTOOLS_API FUVEditorUVDistributeOp : public FUVEditorUVTransformBas
 public:
 	virtual ~FUVEditorUVDistributeOp() {}
 
-	EUVEditorDistributeModeBackend DistributeMode = EUVEditorDistributeModeBackend::TopEdges;
+	EUVEditorDistributeModeBackend DistributeMode = EUVEditorDistributeModeBackend::None;
+	bool bEnableManualDistances;
+	float ManualExtent;
+	float ManualSpacing;
 
 protected:
 	virtual void HandleTransformationOp(FProgressCancel* Progress) override;
@@ -327,10 +425,13 @@ public:
 	// IDynamicMeshOperatorFactory API
 	virtual TUniquePtr<UE::Geometry::FDynamicMeshOperator> MakeNewOperator() override;
 
+	EUVEditorUVTransformType TransformType;
+
 	UPROPERTY()
-	TObjectPtr<UUVEditorUVTransformProperties> Settings;
+	TObjectPtr<UUVEditorUVTransformPropertiesBase> Settings;
 
-	TOptional<TSet<int32>> Selection;
+	TOptional<TSet<int32>> TriangleSelection;
+	TOptional<TSet<int32>> VertexSelection;
 
 	TSharedPtr<UE::Geometry::FDynamicMesh3, ESPMode::ThreadSafe> OriginalMesh;
 	TUniqueFunction<int32()> GetSelectedUVChannel = []() { return 0; };
diff --git a/Engine/Plugins/Experimental/UVEditor/Source/UVEditorTools/Public/UVEditorStyle.h b/Engine/Plugins/Experimental/UVEditor/Source/UVEditorTools/Public/UVEditorStyle.h
new file mode 100644
--- /dev/null
+++ b/Engine/Plugins/Experimental/UVEditor/Source/UVEditorTools/Public/UVEditorStyle.h
@@ -0,0 +1,24 @@
+// Copyright Epic Games, Inc. All Rights Reserved.
+
+#pragma once
+
+#include "Styling/SlateStyle.h"
+#include "Styling/SlateStyleRegistry.h"
+
+/**
+ * Slate style set for UV Editor
+ */
+class UVEDITORTOOLS_API FUVEditorStyle
+	: public FSlateStyleSet
+{
+public:
+	static FName StyleName;
+
+	/** Access the singleton instance for this style set */
+	static FUVEditorStyle& Get();
+
+private:
+
+	FUVEditorStyle();
+	~FUVEditorStyle();
+};
\ No newline at end of file
diff --git a/Engine/Plugins/Experimental/UVEditor/Source/UVEditorTools/Public/UVEditorTransformTool.h b/Engine/Plugins/Experimental/UVEditor/Source/UVEditorTools/Public/UVEditorTransformTool.h
--- a/Engine/Plugins/Experimental/UVEditor/Source/UVEditorTools/Public/UVEditorTransformTool.h
+++ b/Engine/Plugins/Experimental/UVEditor/Source/UVEditorTools/Public/UVEditorTransformTool.h
@@ -8,6 +8,7 @@
 #include "InteractiveToolBuilder.h"
 #include "UVEditorToolAnalyticsUtils.h"
 #include "Selection/UVToolSelectionAPI.h"
+#include "Operators/UVEditorUVTransformOp.h"
 
 #include "UVEditorTransformTool.generated.h"
 
@@ -62,28 +63,6 @@ protected:
 	virtual void ConfigureTool(UUVEditorTransformTool* NewTool) const  override;
 };
 
-/**
- * UV Quick Transform Settings
- */
-UCLASS()
-class UVEDITORTOOLS_API UUVEditorUVQuickTransformProperties : public UInteractiveToolPropertySet
-{
-	GENERATED_BODY()
-
-public:
-	UPROPERTY(EditAnywhere, Category = "Quick Transform", meta = (TransientToolProperty, DisplayName = "Translation"))
-	float QuickTranslateOffset = 0.0;
-
-	UPROPERTY(EditAnywhere, Category = "Quick Transform", meta = (TransientToolProperty, DisplayName = "Rotation"))
-	float QuickRotationOffset = 0.0;
-
-	// parent ref required for details customization
-	UPROPERTY(meta = (TransientToolProperty))
-	TWeakObjectPtr<UUVEditorTransformTool> Tool;
-
-};
-
-
 /**
  * 
  */
@@ -117,9 +96,6 @@ public:
 
 	virtual void OnPropertyModified(UObject* PropertySet, FProperty* Property) override;
 
-	void InitiateQuickTranslate(float Offset, const FVector2D& Direction);
-	void InitiateQuickRotation(float Rotation);
-
 protected:	
 
 	TOptional<EUVEditorUVTransformType> ToolMode;
@@ -128,10 +104,7 @@ protected:
 	TArray<TObjectPtr<UUVEditorToolMeshInput>> Targets;
 
 	UPROPERTY()
-	TObjectPtr<UUVEditorUVQuickTransformProperties> QuickTransformSettings = nullptr;
-
-	UPROPERTY()
-	TObjectPtr<UUVEditorUVTransformProperties> Settings = nullptr;
+	TObjectPtr<UUVEditorUVTransformPropertiesBase> Settings = nullptr;
 
 	UPROPERTY()
 	TArray<TObjectPtr<UUVEditorUVTransformOperatorFactory>> Factories;
diff --git a/Engine/Plugins/Experimental/UVEditor/Source/UVEditorToolsEditorOnly/Private/DetailsCustomizations/UVTransformToolCustomizations.cpp b/Engine/Plugins/Experimental/UVEditor/Source/UVEditorToolsEditorOnly/Private/DetailsCustomizations/UVTransformToolCustomizations.cpp
--- a/Engine/Plugins/Experimental/UVEditor/Source/UVEditorToolsEditorOnly/Private/DetailsCustomizations/UVTransformToolCustomizations.cpp
+++ b/Engine/Plugins/Experimental/UVEditor/Source/UVEditorToolsEditorOnly/Private/DetailsCustomizations/UVTransformToolCustomizations.cpp
@@ -8,6 +8,8 @@
 #include "Widgets/Input/SButton.h"
 #include "Widgets/Input/SComboButton.h"
 #include "Widgets/Layout/SUniformGridPanel.h"
+#include "Widgets/Layout/SUniformWrapPanel.h"
+#include "Widgets/Images/SLayeredImage.h"
 #include "IDetailChildrenBuilder.h"
 #include "Internationalization/BreakIterator.h"
 #include "Framework/MultiBox/MultiBoxBuilder.h"
@@ -16,12 +18,14 @@
 
 #include "PropertyRestriction.h"
 
+#include "UVEditorStyle.h"
 #include "UVEditorTransformTool.h"
+#include "Operators/UVEditorUVTransformOp.h"
 
 
 #define LOCTEXT_NAMESPACE "UVEditorDetailsCustomization"
 
-namespace UVEditorDetailsCustomizationLocal
+namespace UVEditorTransformDetailsCustomizationLocal
 {
 	void CustomSortTransformToolCategories(const TMap<FName, IDetailCategoryBuilder*>&  AllCategoryMap )
 	{
@@ -29,6 +33,66 @@ namespace UVEditorDetailsCustomizationLocal
 		(*AllCategoryMap.Find(FName("Quick Rotate")))->SetSortOrder(1);
 		(*AllCategoryMap.Find(FName("Quick Transform")))->SetSortOrder(2);
 	}
+
+	template<class UENUM_TYPE>
+	FPropertyAccess::Result GetPropertyValueAsEnum(const TSharedPtr<IPropertyHandle> Property, UENUM_TYPE& Value)
+	{
+		FPropertyAccess::Result Result;
+		if (Property.IsValid())
+		{
+			uint8 ValueAsInt8;
+			Result = Property->GetValue(/*out*/ ValueAsInt8);
+			if (Result == FPropertyAccess::Success)
+			{
+				Value = (UENUM_TYPE)ValueAsInt8;
+				return FPropertyAccess::Success;
+			}
+			uint16 ValueAsInt16;
+			Result = Property->GetValue(/*out*/ ValueAsInt16);
+			if (Result == FPropertyAccess::Success)
+			{
+				Value = (UENUM_TYPE)ValueAsInt16;
+				return FPropertyAccess::Success;
+			}
+			uint32 ValueAsInt32;
+			Result = Property->GetValue(/*out*/ ValueAsInt32);
+			if (Result == FPropertyAccess::Success)
+			{
+				Value = (UENUM_TYPE)ValueAsInt32;
+				return FPropertyAccess::Success;
+			}
+		}
+		return FPropertyAccess::Fail;
+	}
+
+	template<class UENUM_TYPE>
+	FPropertyAccess::Result SetPropertyValueAsEnum(const TSharedPtr<IPropertyHandle> Property, const UENUM_TYPE& Value)
+	{
+		FPropertyAccess::Result Result;
+
+		if (Property.IsValid())
+		{
+			uint32 ValueAsInt32 = (uint32)Value;;
+			Result = Property->SetValue(ValueAsInt32);
+			if (Result == FPropertyAccess::Success)
+			{
+				return Result;
+			}
+			uint16 ValueAsInt16 = (uint16)Value;;
+			Result = Property->SetValue(ValueAsInt16);
+			if (Result == FPropertyAccess::Success)
+			{
+				return Result;
+			}
+			uint8 ValueAsInt8 = (uint8)Value;;
+			Result = Property->SetValue(ValueAsInt8);
+			if (Result == FPropertyAccess::Success)
+			{
+				return Result;
+			}
+		}
+		return FPropertyAccess::Fail;
+	}
 }
 
 
@@ -45,15 +109,11 @@ TSharedRef<IDetailCustomization> FUVEditorUVTransformToolDetails::MakeInstance()
 
 void FUVEditorUVTransformToolDetails::CustomizeDetails(IDetailLayoutBuilder& DetailBuilder)
 {
-	using namespace UVEditorDetailsCustomizationLocal;
+	using namespace UVEditorTransformDetailsCustomizationLocal;
 
 	TArray<TWeakObjectPtr<UObject>> ObjectsBeingCustomized;
 	DetailBuilder.GetObjectsBeingCustomized(ObjectsBeingCustomized);
 	check(ObjectsBeingCustomized.Num() > 0);
-	UUVEditorUVQuickTransformProperties* QuickTransformProperties = CastChecked<UUVEditorUVQuickTransformProperties>(ObjectsBeingCustomized[0]);
-	UUVEditorTransformTool* Tool = QuickTransformProperties->Tool.Get();
-	TargetTool = Tool;
-
 	BuildQuickTranslateMenu(DetailBuilder);
 	BuildQuickRotateMenu(DetailBuilder);
 
@@ -65,15 +125,22 @@ void FUVEditorUVTransformToolDetails::BuildQuickTranslateMenu(IDetailLayoutBuild
 	const FName QuickTranslateCategoryName = "Quick Translate";
 	IDetailCategoryBuilder& CategoryBuilder = DetailBuilder.EditCategory(QuickTranslateCategoryName);	
 
-	TSharedPtr<IPropertyHandle> QuickTranslateOffsetHandle = DetailBuilder.GetProperty(GET_MEMBER_NAME_CHECKED(UUVEditorUVQuickTransformProperties, QuickTranslateOffset), UUVEditorUVQuickTransformProperties::StaticClass());
+	TSharedPtr<IPropertyHandle> QuickTranslateOffsetHandle = DetailBuilder.GetProperty(GET_MEMBER_NAME_CHECKED(UUVEditorUVTransformProperties, QuickTranslateOffset), UUVEditorUVTransformProperties::StaticClass());
 	ensure(QuickTranslateOffsetHandle->IsValidHandle());
 	QuickTranslateOffsetHandle->MarkHiddenByCustomization();
 
-	auto ApplyTranslation = [this, QuickTranslateOffsetHandle](const FVector2D& Direction)
+	TSharedPtr<IPropertyHandle> QuickTranslateHandle = DetailBuilder.GetProperty(GET_MEMBER_NAME_CHECKED(UUVEditorUVTransformProperties, QuickTranslation), UUVEditorUVTransformProperties::StaticClass());
+	ensure(QuickTranslateHandle->IsValidHandle());
+	QuickTranslateHandle->MarkHiddenByCustomization();
+
+	auto ApplyTranslation = [this, QuickTranslateHandle, QuickTranslateOffsetHandle](const FVector2D& Direction)
 	{
 		float TranslationValue;
+		FVector2D TotalTranslationValue;
 		QuickTranslateOffsetHandle->GetValue(TranslationValue);
-		return OnQuickTranslate(TranslationValue, Direction);
+		QuickTranslateHandle->GetValue(TotalTranslationValue);
+		QuickTranslateHandle->SetValue(TotalTranslationValue + (Direction * TranslationValue) );
+		return FReply::Handled();
 	};
 
 	FDetailWidgetRow& CustomRow = CategoryBuilder.AddCustomRow(FText::GetEmpty())
@@ -87,7 +154,7 @@ void FUVEditorUVTransformToolDetails::BuildQuickTranslateMenu(IDetailLayoutBuild
 			SNew(SButton)
 			.Text(LOCTEXT("QuickMoveTopLeft", "TL"))
 			.ToolTipText(LOCTEXT("QuickMoveTopLeftToolTip", "Applies translation offset in the negative X axis and the positive Y axis"))
-			.OnClicked_Lambda([ApplyTranslation]() { return ApplyTranslation(FVector2D(-1.0, 1.0)); })			
+			.OnClicked_Lambda([ApplyTranslation]() { return ApplyTranslation(FVector2D(-1.0, 1.0)); })		
 		]
 		+ SUniformGridPanel::Slot(1, 0)
 		.HAlign(HAlign_Center)
@@ -170,10 +237,22 @@ void FUVEditorUVTransformToolDetails::BuildQuickRotateMenu(IDetailLayoutBuilder&
 	const FName QuickRotateCategoryName = "Quick Rotate";
 	IDetailCategoryBuilder& CategoryBuilder = DetailBuilder.EditCategory(QuickRotateCategoryName);
 
-	TSharedPtr<IPropertyHandle> QuickRotationOffsetHandle = DetailBuilder.GetProperty(GET_MEMBER_NAME_CHECKED(UUVEditorUVQuickTransformProperties, QuickRotationOffset), UUVEditorUVQuickTransformProperties::StaticClass());
+	TSharedPtr<IPropertyHandle> QuickRotationOffsetHandle = DetailBuilder.GetProperty(GET_MEMBER_NAME_CHECKED(UUVEditorUVTransformProperties, QuickRotationOffset), UUVEditorUVTransformProperties::StaticClass());
 	ensure(QuickRotationOffsetHandle->IsValidHandle());
 	QuickRotationOffsetHandle->MarkHiddenByCustomization();
 
+	TSharedPtr<IPropertyHandle> QuickRotationHandle = DetailBuilder.GetProperty(GET_MEMBER_NAME_CHECKED(UUVEditorUVTransformProperties, QuickRotation), UUVEditorUVTransformProperties::StaticClass());
+	ensure(QuickRotationHandle->IsValidHandle());
+	QuickRotationHandle->MarkHiddenByCustomization();
+
+	auto ApplyRotation = [this, QuickRotationHandle](const float& RotationOffset)
+	{		
+		float TotalRotationValue;		
+		QuickRotationHandle->GetValue(TotalRotationValue);
+		QuickRotationHandle->SetValue(TotalRotationValue + RotationOffset);
+		return FReply::Handled();
+	};
+
 	FDetailWidgetRow& CustomRow = CategoryBuilder.AddCustomRow(FText::GetEmpty())
 	[
 		SNew(SUniformGridPanel)
@@ -185,10 +264,10 @@ void FUVEditorUVTransformToolDetails::BuildQuickRotateMenu(IDetailLayoutBuilder&
 			SNew(SButton)
 			.Text(LOCTEXT("QuickRotateClockwise", "CW"))
 			.ToolTipText(LOCTEXT("QuickRotateClockwiseToolTip", "Applies the rotation in a clockwise orientation"))
-			.OnClicked_Lambda([this, QuickRotationOffsetHandle]() {
+			.OnClicked_Lambda([this, QuickRotationOffsetHandle, ApplyRotation]() {
 				float RotationValue;
 				QuickRotationOffsetHandle->GetValue(RotationValue);
-				return OnQuickRotate(-RotationValue); 
+				return ApplyRotation(-RotationValue);
 			})
 		]
 		+ SUniformGridPanel::Slot(1, 0)
@@ -205,10 +284,10 @@ void FUVEditorUVTransformToolDetails::BuildQuickRotateMenu(IDetailLayoutBuilder&
 			SNew(SButton)
 			.Text(LOCTEXT("QuickRotateCounterclockwise", "CCW"))
 			.ToolTipText(LOCTEXT("QuickRotateCounterclockwiseToolTip", "Applies the rotation in a counter clockwise orientation"))
-			.OnClicked_Lambda([this, QuickRotationOffsetHandle]() {
+			.OnClicked_Lambda([this, QuickRotationOffsetHandle, ApplyRotation]() {
 				float RotationValue;
 				QuickRotationOffsetHandle->GetValue(RotationValue);
-				return OnQuickRotate(RotationValue); 
+				return ApplyRotation(RotationValue);
 			})
 		]
 
@@ -219,7 +298,7 @@ void FUVEditorUVTransformToolDetails::BuildQuickRotateMenu(IDetailLayoutBuilder&
 			SNew(SButton)
 			.Text(LOCTEXT("QuickRotateClockwise10Deg", "10°"))
 			.ToolTipText(LOCTEXT("QuickRotateClockwise10DegToolTip", "Applies a 10 degree clockwise orientation"))
-			.OnClicked_Lambda([this]() { return OnQuickRotate(-10.0); })
+			.OnClicked_Lambda([this, ApplyRotation]() { return ApplyRotation(-10.0); })
 		]
 		+ SUniformGridPanel::Slot(2, 1)
 		.HAlign(HAlign_Center)
@@ -228,7 +307,7 @@ void FUVEditorUVTransformToolDetails::BuildQuickRotateMenu(IDetailLayoutBuilder&
 			SNew(SButton)
 			.Text(LOCTEXT("QuickRotateCounterclockwise10Deg", "10°"))
 			.ToolTipText(LOCTEXT("QuickRotateCounterclockwise10DegToolTip", "Applies a 10 degree counter clockwise orientation"))
-			.OnClicked_Lambda([this]() { return OnQuickRotate(10.0); })
+			.OnClicked_Lambda([this, ApplyRotation]() { return ApplyRotation(10.0); })
 		]
 		+ SUniformGridPanel::Slot(0, 2)
 		.HAlign(HAlign_Center)
@@ -237,7 +316,7 @@ void FUVEditorUVTransformToolDetails::BuildQuickRotateMenu(IDetailLayoutBuilder&
 			SNew(SButton)
 			.Text(LOCTEXT("QuickRotateClockwise45Deg", "45°"))
 			.ToolTipText(LOCTEXT("QuickRotateClockwise45DegToolTip", "Applies a 45 degree clockwise orientation"))
-			.OnClicked_Lambda([this]() { return OnQuickRotate(-45.0); })
+			.OnClicked_Lambda([this, ApplyRotation]() { return ApplyRotation(-45.0); })
 		]
 		+ SUniformGridPanel::Slot(2, 2)
 		.HAlign(HAlign_Center)
@@ -246,7 +325,7 @@ void FUVEditorUVTransformToolDetails::BuildQuickRotateMenu(IDetailLayoutBuilder&
 			SNew(SButton)
 			.Text(LOCTEXT("QuickRotateCounterclockwise45Deg", "45°"))
 			.ToolTipText(LOCTEXT("QuickRotateCounterclockwise45DegToolTip", "Applies a 45 degree counter clockwise orientation"))
-			.OnClicked_Lambda([this]() { return OnQuickRotate(45.0); })
+			.OnClicked_Lambda([this, ApplyRotation]() { return ApplyRotation(45.0); })
 		]
 		+ SUniformGridPanel::Slot(0, 3)
 		.HAlign(HAlign_Center)
@@ -255,7 +334,7 @@ void FUVEditorUVTransformToolDetails::BuildQuickRotateMenu(IDetailLayoutBuilder&
 			SNew(SButton)
 			.Text(LOCTEXT("QuickRotateClockwise90Deg", "90°"))
 			.ToolTipText(LOCTEXT("QuickRotateClockwise90DegToolTip", "Applies a 90 degree clockwise orientation"))
-			.OnClicked_Lambda([this]() { return OnQuickRotate(-90.0); })
+			.OnClicked_Lambda([this, ApplyRotation]() { return ApplyRotation(-90.0); })
 		]
 		+ SUniformGridPanel::Slot(2, 3)
 		.HAlign(HAlign_Center)
@@ -264,24 +343,510 @@ void FUVEditorUVTransformToolDetails::BuildQuickRotateMenu(IDetailLayoutBuilder&
 			SNew(SButton)
 			.Text(LOCTEXT("QuickRotateCounterclockwise90Deg", "90°"))
 			.ToolTipText(LOCTEXT("QuickRotateCounterclockwise90DegToolTip", "Applies a 90 degree counter clockwise orientation"))
-			.OnClicked_Lambda([this]() { return OnQuickRotate(90.0); })
+			.OnClicked_Lambda([this, ApplyRotation]() { return ApplyRotation(90.0); })
 		]
 	];
 }
 
+//
+// UVEditorTransformTool - Quick Transform version
+//
+
 
-FReply  FUVEditorUVTransformToolDetails::OnQuickTranslate(float TranslationValue, const FVector2D& Direction)
+TSharedRef<IDetailCustomization> FUVEditorUVQuickTransformToolDetails::MakeInstance()
 {
-	ensure(TargetTool.IsValid());
-	TargetTool->InitiateQuickTranslate(TranslationValue, Direction);
-	return FReply::Handled();
+	return MakeShareable(new FUVEditorUVTransformToolDetails);
 }
 
-FReply  FUVEditorUVTransformToolDetails::OnQuickRotate(float RotationValue)
+
+void FUVEditorUVQuickTransformToolDetails::CustomizeDetails(IDetailLayoutBuilder& DetailBuilder)
 {
-	ensure(TargetTool.IsValid());
-	TargetTool->InitiateQuickRotation(RotationValue);
-	return FReply::Handled();
+	using namespace UVEditorTransformDetailsCustomizationLocal;
+
+	TArray<TWeakObjectPtr<UObject>> ObjectsBeingCustomized;
+	DetailBuilder.GetObjectsBeingCustomized(ObjectsBeingCustomized);
+	check(ObjectsBeingCustomized.Num() > 0);
+
+	DetailBuilder.HideCategory("Advanced Transform");
+
+	BuildQuickTranslateMenu(DetailBuilder);
+	BuildQuickRotateMenu(DetailBuilder);
+
+	DetailBuilder.SortCategories(&CustomSortTransformToolCategories);
 }
 
+//
+// UVEditorDistributeTool
+//
+
+
+TSharedRef<IDetailCustomization> FUVEditorUVDistributeToolDetails::MakeInstance()
+{
+	return MakeShareable(new FUVEditorUVDistributeToolDetails);
+}
+
+
+void FUVEditorUVDistributeToolDetails::CustomizeDetails(IDetailLayoutBuilder& DetailBuilder)
+{
+	using namespace UVEditorTransformDetailsCustomizationLocal;
+
+	TArray<TWeakObjectPtr<UObject>> ObjectsBeingCustomized;
+	DetailBuilder.GetObjectsBeingCustomized(ObjectsBeingCustomized);
+	check(ObjectsBeingCustomized.Num() > 0);
+	BuildDistributeModeButtons(DetailBuilder);	
+}
+
+void FUVEditorUVDistributeToolDetails::BuildDistributeModeButtons(IDetailLayoutBuilder& DetailBuilder)
+{
+	using namespace UVEditorTransformDetailsCustomizationLocal;
+
+	const FName QuickRotateCategoryName = "Distribute";
+	IDetailCategoryBuilder& CategoryBuilder = DetailBuilder.EditCategory(QuickRotateCategoryName);
+
+	DistributeModeHandle = DetailBuilder.GetProperty(GET_MEMBER_NAME_CHECKED(UUVEditorUVDistributeProperties, DistributeMode), UUVEditorUVDistributeProperties::StaticClass());
+	ensure(DistributeModeHandle->IsValidHandle());
+	DistributeModeHandle->MarkHiddenByCustomization();
+
+	TSharedPtr<IPropertyHandle> GroupingHandle = DetailBuilder.GetProperty(GET_MEMBER_NAME_CHECKED(UUVEditorUVDistributeProperties, Grouping), UUVEditorUVDistributeProperties::StaticClass());
+	ensure(GroupingHandle->IsValidHandle());
+
+	EnableManualDistancesHandle = DetailBuilder.GetProperty(GET_MEMBER_NAME_CHECKED(UUVEditorUVDistributeProperties, bEnableManualDistances), UUVEditorUVDistributeProperties::StaticClass());
+	ensure(EnableManualDistancesHandle->IsValidHandle());
+
+	OrigManualExtentHandle = DetailBuilder.GetProperty(GET_MEMBER_NAME_CHECKED(UUVEditorUVDistributeProperties, ManualExtent), UUVEditorUVDistributeProperties::StaticClass());
+	ensure(OrigManualExtentHandle->IsValidHandle());
+
+	OrigManualSpacingHandle = DetailBuilder.GetProperty(GET_MEMBER_NAME_CHECKED(UUVEditorUVDistributeProperties, ManualSpacing), UUVEditorUVDistributeProperties::StaticClass());
+	ensure(OrigManualSpacingHandle->IsValidHandle());
+
+
+	auto BuildButton = [this](const FName& Icon, const FText& Tooltip, TFunction<FReply()> OnClickCallback)
+	{
+
+		TSharedRef<SLayeredImage> IconWidget =
+			SNew(SLayeredImage)
+			.Visibility(EVisibility::HitTestInvisible)
+			.Image(FUVEditorStyle::Get().GetBrush(Icon));
+
+		TSharedRef<SWidget> ButtonContent = SNullWidget::NullWidget;
+
+		ButtonContent =
+			SNew(SHorizontalBox)
+			+ SHorizontalBox::Slot()
+			.FillWidth(1)
+			.Padding(0)
+			.VAlign(VAlign_Center)
+			[
+				SNew(SVerticalBox)
+				// Icon image
+				+ SVerticalBox::Slot()
+				.FillHeight(1)
+				.Padding(0)
+				.HAlign(HAlign_Center)	// Center the icon horizontally, so that large labels don't stretch out the artwork
+				[
+					IconWidget
+				]
+			];
+
+
+		TSharedRef<SWidget> Button = SNew(SButton)
+			.ButtonStyle(&FAppStyle::Get().GetWidgetStyle<FButtonStyle>("Button"))
+			.ContentPadding(FMargin(0, 5.0))
+			.ContentScale(FVector2D(1, 1))
+			.ToolTipText(Tooltip)
+			.OnClicked_Lambda(OnClickCallback)
+			[
+				ButtonContent
+			];
+		return Button;
+	};
+
+	FDetailWidgetRow& CustomRow = CategoryBuilder.AddCustomRow(FText::GetEmpty())
+		[
+			SNew(SVerticalBox)
+			+ SVerticalBox::Slot()
+			.AutoHeight()
+			[
+				SNew(SUniformWrapPanel)
+				.SlotPadding(FMargin(5.0f))
+				.EvenRowDistribution(true)
+				+ SUniformWrapPanel::Slot()
+					.HAlign(HAlign_Center)
+					.VAlign(VAlign_Center)
+					[
+						BuildButton("UVEditor.DistributeSpaceVertically",
+						LOCTEXT("DistributeVerticalSpaceToolTip", "Distributes objects so that vertical space between them is equal."),
+						[this]()
+							{
+								SetPropertyValueAsEnum(DistributeModeHandle, EUVEditorDistributeMode::VerticalSpace);
+								return FReply::Handled();
+							}
+						)
+					]
+				+ SUniformWrapPanel::Slot()
+					.HAlign(HAlign_Center)
+					.VAlign(VAlign_Center)
+					[
+						BuildButton("UVEditor.DistributeSpaceHorizontally",
+						LOCTEXT("DistributeHorizontalSpaceToolTip", "Distributes objects so that horizontal space between them is equal."),
+						[this]()
+							{
+								SetPropertyValueAsEnum(DistributeModeHandle, EUVEditorDistributeMode::HorizontalSpace);
+								return FReply::Handled();
+							}
+						)
+					]
+			]
+		+ SVerticalBox::Slot()
+		.AutoHeight()
+		[
+			SNew(SUniformWrapPanel)
+			.SlotPadding(FMargin(5.0f))
+			.EvenRowDistribution(true)
+			+ SUniformWrapPanel::Slot()
+				.HAlign(HAlign_Center)
+				.VAlign(VAlign_Center)
+				[
+					BuildButton("UVEditor.DistributeTopEdges",
+					LOCTEXT("DistributeTopEdgesToolTip", "Distributes objects so that distance between top edges is equal."),
+					[this]()
+						{
+							SetPropertyValueAsEnum(DistributeModeHandle, EUVEditorDistributeMode::TopEdges);
+							return FReply::Handled();
+						}
+					)
+				]
+			+ SUniformWrapPanel::Slot()
+				.HAlign(HAlign_Center)
+				.VAlign(VAlign_Center)
+				[
+					BuildButton("UVEditor.DistributeBottomEdges",
+					LOCTEXT("DistributeBottomEdgesToolTip", "Distributes objects so that distance between bottom edges is equal."),
+					[this]()
+						{
+							SetPropertyValueAsEnum(DistributeModeHandle, EUVEditorDistributeMode::BottomEdges);
+							return FReply::Handled();
+						}
+					)
+				]
+			+ SUniformWrapPanel::Slot()
+				.HAlign(HAlign_Center)
+				.VAlign(VAlign_Center)
+				[
+					BuildButton("UVEditor.DistributeLeftEdges",
+					LOCTEXT("DistributeLeftEdgesToolTip", "Distributes objects so that distance between left edges is equal."),
+					[this]()
+						{
+							SetPropertyValueAsEnum(DistributeModeHandle, EUVEditorDistributeMode::LeftEdges);
+							return FReply::Handled();
+						}
+					)
+				]
+			+ SUniformWrapPanel::Slot()
+				.HAlign(HAlign_Center)
+				.VAlign(VAlign_Center)
+				[
+					BuildButton("UVEditor.DistributeRightEdges",
+					LOCTEXT("DistributeRightEdgesToolTip", "Distributes objects so that distance between right edges is equal."),
+					[this]()
+						{
+							SetPropertyValueAsEnum(DistributeModeHandle, EUVEditorDistributeMode::RightEdges);
+							return FReply::Handled();
+						}
+					)
+				]
+			+ SUniformWrapPanel::Slot()
+				.HAlign(HAlign_Center)
+				.VAlign(VAlign_Center)
+				[
+					BuildButton("UVEditor.DistributeCentersHorizontally",
+					LOCTEXT("DistributeCenterHorizontalToolTip", "Distributes objects so that distance between horizontal centers is equal."),
+					[this]()
+						{
+							SetPropertyValueAsEnum(DistributeModeHandle, EUVEditorDistributeMode::CentersHorizontally);
+							return FReply::Handled();
+						}
+					)
+				]
+			+ SUniformWrapPanel::Slot()
+				.HAlign(HAlign_Center)
+				.VAlign(VAlign_Center)
+				[
+					BuildButton("UVEditor.DistributeCentersVertically",
+					LOCTEXT("DistributeCenterVerticalToolTip", "Distributes objects so that distance between vertical centers is equal."),
+					[this]()
+						{
+							SetPropertyValueAsEnum(DistributeModeHandle, EUVEditorDistributeMode::CentersVertically);
+							return FReply::Handled();
+						}
+					)
+				]
+			]
+		+ SVerticalBox::Slot()
+		.AutoHeight()
+		[
+			SNew(SUniformWrapPanel)
+			.SlotPadding(FMargin(5.0f))
+			.EvenRowDistribution(true)
+			+ SUniformWrapPanel::Slot()
+				.HAlign(HAlign_Center)
+				.VAlign(VAlign_Center)
+				[
+					BuildButton("UVEditor.DistributeRemoveOverlap",
+					LOCTEXT("DistributeMinimalOverlapToolTip", "Distributes objects so that overlap is removed, moving objects minimally as possible."),
+					[this]()
+						{
+							SetPropertyValueAsEnum(DistributeModeHandle, EUVEditorDistributeMode::MinimallyRemoveOverlap);
+							return FReply::Handled();
+						}
+					)
+				]
+		]
+	];
+
+	
+	static FText RestrictReason = LOCTEXT("GroupingRestrictionText", "Not supported in Distribute Tool.");
+	TSharedPtr<FPropertyRestriction> GroupingEnumRestriction = MakeShareable(new FPropertyRestriction(RestrictReason));
+	GroupingHandle->AddRestriction(GroupingEnumRestriction.ToSharedRef());
+	GroupingEnumRestriction->AddHiddenValue("EnclosingBoundingBox");
+
+	CategoryBuilder.AddProperty(GroupingHandle);
+
+	CategoryBuilder.AddProperty(EnableManualDistancesHandle);
+	IDetailPropertyRow& ManualSpacingRow = CategoryBuilder.AddProperty(OrigManualSpacingHandle);
+	IDetailPropertyRow& ManualExtentRow = CategoryBuilder.AddProperty(OrigManualExtentHandle);
+
+	ManualSpacingRow.Visibility(TAttribute<EVisibility>::CreateLambda([this]()
+		{
+			bool bManualDistance;
+			EnableManualDistancesHandle->GetValue(bManualDistance);
+
+			EUVEditorDistributeMode DistributeMode;
+			GetPropertyValueAsEnum(DistributeModeHandle, DistributeMode);
+
+			EVisibility Visibility;
+
+			if (bManualDistance)
+			{
+				if (DistributeMode == EUVEditorDistributeMode::VerticalSpace ||
+					DistributeMode == EUVEditorDistributeMode::HorizontalSpace ||
+					DistributeMode == EUVEditorDistributeMode::MinimallyRemoveOverlap)
+				{
+					return EVisibility::Visible;
+				}
+			}
+			return EVisibility::Collapsed;
+		}));
+
+	ManualExtentRow.Visibility(TAttribute<EVisibility>::CreateLambda([this]()
+		{
+			bool bManualDistance;
+			EnableManualDistancesHandle->GetValue(bManualDistance);
+
+			EUVEditorDistributeMode DistributeMode;
+			GetPropertyValueAsEnum(DistributeModeHandle, DistributeMode);
+
+			if (bManualDistance)
+			{
+				if (DistributeMode == EUVEditorDistributeMode::TopEdges ||
+					DistributeMode == EUVEditorDistributeMode::BottomEdges ||
+					DistributeMode == EUVEditorDistributeMode::LeftEdges ||
+					DistributeMode == EUVEditorDistributeMode::RightEdges ||
+					DistributeMode == EUVEditorDistributeMode::CentersHorizontally ||
+					DistributeMode == EUVEditorDistributeMode::CentersVertically)
+				{
+					return EVisibility::Visible;
+				}
+			}
+			return EVisibility::Collapsed;
+		}));
+
+}
+
+//
+// UVEditorAlignTool
+//
+
+TSharedRef<IDetailCustomization> FUVEditorUVAlignToolDetails::MakeInstance()
+{
+	return MakeShareable(new FUVEditorUVAlignToolDetails);
+}
+
+
+void FUVEditorUVAlignToolDetails::CustomizeDetails(IDetailLayoutBuilder& DetailBuilder)
+{
+	using namespace UVEditorTransformDetailsCustomizationLocal;
+
+	TArray<TWeakObjectPtr<UObject>> ObjectsBeingCustomized;
+	DetailBuilder.GetObjectsBeingCustomized(ObjectsBeingCustomized);
+	check(ObjectsBeingCustomized.Num() > 0);
+	BuildAlignModeButtons(DetailBuilder);
+}
+
+void FUVEditorUVAlignToolDetails::BuildAlignModeButtons(IDetailLayoutBuilder& DetailBuilder)
+{
+	using namespace UVEditorTransformDetailsCustomizationLocal;
+
+	const FName QuickRotateCategoryName = "Align";
+	IDetailCategoryBuilder& CategoryBuilder = DetailBuilder.EditCategory(QuickRotateCategoryName);
+
+	AlignDirectionHandle = DetailBuilder.GetProperty(GET_MEMBER_NAME_CHECKED(UUVEditorUVAlignProperties, AlignDirection), UUVEditorUVAlignProperties::StaticClass());
+	ensure(AlignDirectionHandle->IsValidHandle());
+	AlignDirectionHandle->MarkHiddenByCustomization();
+
+	TSharedPtr<IPropertyHandle> GroupingHandle = DetailBuilder.GetProperty(GET_MEMBER_NAME_CHECKED(UUVEditorUVAlignProperties, Grouping), UUVEditorUVAlignProperties::StaticClass());
+	ensure(GroupingHandle->IsValidHandle());
+
+	TSharedPtr<IPropertyHandle> AlignAnchorHandle = DetailBuilder.GetProperty(GET_MEMBER_NAME_CHECKED(UUVEditorUVAlignProperties, AlignAnchor), UUVEditorUVAlignProperties::StaticClass());
+	ensure(AlignAnchorHandle->IsValidHandle());
+
+	TSharedPtr<IPropertyHandle> ManualAnchorHandle = DetailBuilder.GetProperty(GET_MEMBER_NAME_CHECKED(UUVEditorUVAlignProperties, ManualAnchor), UUVEditorUVAlignProperties::StaticClass());
+	ensure(ManualAnchorHandle->IsValidHandle());
+
+	auto BuildButton = [this](const FName& Icon, const FText& Tooltip, TFunction<FReply()> OnClickCallback)
+	{
+
+		TSharedRef<SLayeredImage> IconWidget =
+			SNew(SLayeredImage)
+			.Visibility(EVisibility::HitTestInvisible)
+			.Image(FUVEditorStyle::Get().GetBrush(Icon));
+
+		TSharedRef<SWidget> ButtonContent = SNullWidget::NullWidget;
+
+		ButtonContent =
+			SNew(SHorizontalBox)
+			+ SHorizontalBox::Slot()
+			.FillWidth(1)
+			.Padding(0)
+			.VAlign(VAlign_Center)
+			[
+				SNew(SVerticalBox)
+				// Icon image
+				+ SVerticalBox::Slot()
+				.FillHeight(1)
+				.Padding(0)
+				.HAlign(HAlign_Center)	// Center the icon horizontally, so that large labels don't stretch out the artwork
+				[
+					IconWidget
+				]
+			];
+
+
+		TSharedRef<SWidget> Button = SNew(SButton)
+			.ButtonStyle(&FAppStyle::Get().GetWidgetStyle<FButtonStyle>("Button"))
+			.ContentPadding(FMargin(0, 5.0))
+			.ContentScale(FVector2D(1, 1))
+			.ToolTipText(Tooltip)
+			.OnClicked_Lambda(OnClickCallback)
+			[
+				ButtonContent
+			];
+		return Button;
+	};
+
+
+	FDetailWidgetRow& CustomRow = CategoryBuilder.AddCustomRow(FText::GetEmpty())
+		[
+			SNew(SVerticalBox)
+			+ SVerticalBox::Slot()
+			.AutoHeight()
+			[
+				SNew(SUniformWrapPanel)
+				.SlotPadding(FMargin(5.0f))
+				.EvenRowDistribution(true)
+				+ SUniformWrapPanel::Slot()
+					.HAlign(HAlign_Center)
+					.VAlign(VAlign_Center)
+					[
+						BuildButton("UVEditor.AlignDirectionTopEdges",
+						LOCTEXT("AlignTopEdgesToolTip", "Aligns the top edges of objects with the reference point."),
+						[this]()
+							{
+								SetPropertyValueAsEnum(AlignDirectionHandle, EUVEditorAlignDirection::Top);
+								return FReply::Handled();
+							}
+						)
+					]
+				+ SUniformWrapPanel::Slot()
+					.HAlign(HAlign_Center)
+					.VAlign(VAlign_Center)
+					[
+						BuildButton("UVEditor.AlignDirectionBottomEdges",
+						LOCTEXT("AlignBottomEdgesToolTip", "Aligns the bottom edges of objects with the reference point."),
+						[this]()
+							{
+								SetPropertyValueAsEnum(AlignDirectionHandle, EUVEditorAlignDirection::Bottom);
+								return FReply::Handled();
+							}
+						)
+					]
+				+ SUniformWrapPanel::Slot()
+					.HAlign(HAlign_Center)
+					.VAlign(VAlign_Center)
+					[
+						BuildButton("UVEditor.AlignDirectionLeftEdges",
+						LOCTEXT("AlignLeftEdgesToolTip", "Aligns the left edges of objects with the reference point."),
+						[this]()
+							{
+								SetPropertyValueAsEnum(AlignDirectionHandle, EUVEditorAlignDirection::Left);
+								return FReply::Handled();
+							}
+						)
+					]
+				+ SUniformWrapPanel::Slot()
+					.HAlign(HAlign_Center)
+					.VAlign(VAlign_Center)
+					[
+						BuildButton("UVEditor.AlignDirectionRightEdges",
+						LOCTEXT("AlignRightEdgesToolTip", "Aligns the right edges of objects with the reference point."),
+						[this]()
+							{
+								SetPropertyValueAsEnum(AlignDirectionHandle, EUVEditorAlignDirection::Right);
+								return FReply::Handled();
+							}
+						)
+					]
+			]
+			+ SVerticalBox::Slot()
+			.AutoHeight()
+			[
+				SNew(SUniformWrapPanel)
+				.SlotPadding(FMargin(5.0f))
+				.EvenRowDistribution(true)
+				+ SUniformWrapPanel::Slot()
+					.HAlign(HAlign_Center)
+					.VAlign(VAlign_Center)
+					[
+						BuildButton("UVEditor.AlignDirectionCentersVertically",
+						LOCTEXT("AlignCenterVerticalToolTip", "Aligns the vertical centerline of objects with the reference point."),
+						[this]()
+							{
+								SetPropertyValueAsEnum(AlignDirectionHandle, EUVEditorAlignDirection::CenterVertically);
+								return FReply::Handled();
+							}
+						)
+					]
+				+ SUniformWrapPanel::Slot()
+					.HAlign(HAlign_Center)
+					.VAlign(VAlign_Center)
+					[
+						BuildButton("UVEditor.AlignDirectionCentersHorizontally",
+						LOCTEXT("AlignCenterHorizontalToolTip", "Aligns the horizontal centerline of objects with the reference point."),
+						[this]()
+							{
+								SetPropertyValueAsEnum(AlignDirectionHandle, EUVEditorAlignDirection::CenterHorizontally);
+								return FReply::Handled();
+							}
+						)
+					]
+			]
+		];
+
+
+	CategoryBuilder.AddProperty(AlignAnchorHandle);
+	CategoryBuilder.AddProperty(ManualAnchorHandle);
+	CategoryBuilder.AddProperty(GroupingHandle);
+}
+
+
 #undef LOCTEXT_NAMESPACE
\ No newline at end of file
diff --git a/Engine/Plugins/Experimental/UVEditor/Source/UVEditorToolsEditorOnly/Private/UVEditorToolsEditorOnlyModule.cpp b/Engine/Plugins/Experimental/UVEditor/Source/UVEditorToolsEditorOnly/Private/UVEditorToolsEditorOnlyModule.cpp
--- a/Engine/Plugins/Experimental/UVEditor/Source/UVEditorToolsEditorOnly/Private/UVEditorToolsEditorOnlyModule.cpp
+++ b/Engine/Plugins/Experimental/UVEditor/Source/UVEditorToolsEditorOnly/Private/UVEditorToolsEditorOnlyModule.cpp
@@ -10,6 +10,7 @@
 
 #include "Operators/UVEditorRecomputeUVsOp.h"
 #include "UVEditorTransformTool.h"
+#include "Operators/UVEditorUVTransformOp.h"
 
 #define LOCTEXT_NAMESPACE "FUVEditorToolsEditorOnlyModule"
 
@@ -55,9 +56,26 @@ void FUVEditorToolsEditorOnlyModule::OnPostEngineInit()
 	ClassesToUnregisterOnShutdown.Add(UUVEditorRecomputeUVsToolProperties::StaticClass()->GetFName());
 
 	// Transform
-	PropertyModule.RegisterCustomClassLayout("UVEditorUVQuickTransformProperties", FOnGetDetailCustomizationInstance::CreateStatic(&FUVEditorUVTransformToolDetails::MakeInstance));
+	PropertyModule.RegisterCustomClassLayout(UUVEditorUVTransformProperties::StaticClass()->GetFName(),
+		                                     FOnGetDetailCustomizationInstance::CreateStatic(&FUVEditorUVTransformToolDetails::MakeInstance));
+	ClassesToUnregisterOnShutdown.Add(UUVEditorUVTransformProperties::StaticClass()->GetFName());
+
+	// Quick Transform
+	PropertyModule.RegisterCustomClassLayout(UUVEditorUVQuickTransformProperties::StaticClass()->GetFName(),
+		FOnGetDetailCustomizationInstance::CreateStatic(&FUVEditorUVQuickTransformToolDetails::MakeInstance));
 	ClassesToUnregisterOnShutdown.Add(UUVEditorUVQuickTransformProperties::StaticClass()->GetFName());
 
+	// Align
+	PropertyModule.RegisterCustomClassLayout(UUVEditorUVAlignProperties::StaticClass()->GetFName(),
+		FOnGetDetailCustomizationInstance::CreateStatic(&FUVEditorUVAlignToolDetails::MakeInstance));
+	ClassesToUnregisterOnShutdown.Add(UUVEditorUVAlignProperties::StaticClass()->GetFName());
+
+	// Distribute
+	PropertyModule.RegisterCustomClassLayout(UUVEditorUVDistributeProperties::StaticClass()->GetFName(),
+		FOnGetDetailCustomizationInstance::CreateStatic(&FUVEditorUVDistributeToolDetails::MakeInstance));
+	ClassesToUnregisterOnShutdown.Add(UUVEditorUVDistributeProperties::StaticClass()->GetFName());
+
+
 }
 
 #undef LOCTEXT_NAMESPACE
diff --git a/Engine/Plugins/Experimental/UVEditor/Source/UVEditorToolsEditorOnly/Public/DetailsCustomizations/UVTransformToolCustomizations.h b/Engine/Plugins/Experimental/UVEditor/Source/UVEditorToolsEditorOnly/Public/DetailsCustomizations/UVTransformToolCustomizations.h
--- a/Engine/Plugins/Experimental/UVEditor/Source/UVEditorToolsEditorOnly/Public/DetailsCustomizations/UVTransformToolCustomizations.h
+++ b/Engine/Plugins/Experimental/UVEditor/Source/UVEditorToolsEditorOnly/Public/DetailsCustomizations/UVTransformToolCustomizations.h
@@ -9,7 +9,7 @@
 
 class IPropertyHandle;
 class FPropertyRestriction;
-class UUVEditorTransformTool;
+class IUVEditorTransformToolQuickAction;
 
 // Customization for UUVEditorUVTransformProperties
 class FUVEditorUVTransformToolDetails : public IDetailCustomization
@@ -19,12 +19,48 @@ public:
 	virtual void CustomizeDetails(IDetailLayoutBuilder& DetailBuilder) override;
 
 protected:
-	TWeakObjectPtr<UUVEditorTransformTool> TargetTool;
 	void BuildQuickTranslateMenu(IDetailLayoutBuilder& DetailBuilder);
 	void BuildQuickRotateMenu(IDetailLayoutBuilder& DetailBuilder);
+};
 
-	FReply OnQuickTranslate(float TranslationValue, const FVector2D& Direction);
-	FReply OnQuickRotate(float RotationValue);
+// Customization for UUVEditorUVQuickTransformProperties
+class FUVEditorUVQuickTransformToolDetails : public FUVEditorUVTransformToolDetails
+{
+public:
+	static TSharedRef<IDetailCustomization> MakeInstance();
+	virtual void CustomizeDetails(IDetailLayoutBuilder& DetailBuilder) override;
 };
 
+// Customization for UUVEditorUVDistributeProperties
+class FUVEditorUVDistributeToolDetails : public IDetailCustomization
+{
+public:
+	static TSharedRef<IDetailCustomization> MakeInstance();
+	virtual void CustomizeDetails(IDetailLayoutBuilder& DetailBuilder) override;
+
+protected:
+
+	TSharedPtr<IPropertyHandle> DistributeModeHandle;
+	TSharedPtr<IPropertyHandle> EnableManualDistancesHandle;
+
+	TSharedPtr<IPropertyHandle> OrigManualExtentHandle;
+	TSharedPtr<IPropertyHandle> OrigManualSpacingHandle;
+
+	void BuildDistributeModeButtons(IDetailLayoutBuilder& DetailBuilder);
+
+};
+
+// Customization for UUVEditorUVAlignProperties
+class FUVEditorUVAlignToolDetails : public IDetailCustomization
+{
+public:
+	static TSharedRef<IDetailCustomization> MakeInstance();
+	virtual void CustomizeDetails(IDetailLayoutBuilder& DetailBuilder) override;
+
+protected:
+
+	TSharedPtr<IPropertyHandle> AlignDirectionHandle;
+
+	void BuildAlignModeButtons(IDetailLayoutBuilder& DetailBuilder);
 
+};
\ No newline at end of file
