UVEditor: Provide new Select All command bound to Ctrl-A. Command will select all UVs or deselect all UVs if all UVs are currently selected.

diff --git a/Engine/Plugins/Experimental/UVEditor/Source/UVEditor/Private/UVEditorCommands.cpp b/Engine/Plugins/Experimental/UVEditor/Source/UVEditor/Private/UVEditorCommands.cpp
--- a/Engine/Plugins/Experimental/UVEditor/Source/UVEditor/Private/UVEditorCommands.cpp
+++ b/Engine/Plugins/Experimental/UVEditor/Source/UVEditor/Private/UVEditorCommands.cpp
@@ -50,6 +50,7 @@ void FUVEditorCommands::RegisterCommands()
 	UI_COMMAND(TriangleSelection, "Triangle Selection", "Select triangles", EUserInterfaceActionType::ToggleButton, FInputChord(EKeys::Three));
 	UI_COMMAND(IslandSelection, "Island Selection", "Select islands", EUserInterfaceActionType::ToggleButton, FInputChord(EKeys::Four));
 	UI_COMMAND(FullMeshSelection, "Mesh Selection", "Select meshes", EUserInterfaceActionType::ToggleButton, FInputChord(EKeys::Five));
+	UI_COMMAND(SelectAll, "Select All", "Toggle selecting everything based on current selection mode", EUserInterfaceActionType::None, FInputChord(EKeys::A, EModifierKey::Control));
 
 	UI_COMMAND(EnableOrbitCamera, "Orbit", "Enable orbit camera", EUserInterfaceActionType::ToggleButton, FInputChord());
 	UI_COMMAND(EnableFlyCamera, "Fly", "Enable fly camera", EUserInterfaceActionType::ToggleButton, FInputChord());
diff --git a/Engine/Plugins/Experimental/UVEditor/Source/UVEditor/Private/UVEditorMode.cpp b/Engine/Plugins/Experimental/UVEditor/Source/UVEditor/Private/UVEditorMode.cpp
--- a/Engine/Plugins/Experimental/UVEditor/Source/UVEditor/Private/UVEditorMode.cpp
+++ b/Engine/Plugins/Experimental/UVEditor/Source/UVEditor/Private/UVEditorMode.cpp
@@ -672,6 +672,70 @@ void UUVEditorMode::BindCommands()
 		FIsActionButtonVisible::CreateLambda([this]() {return GetInteractiveToolsContext()->ActiveToolHasAccept(); }),
 		EUIActionRepeatMode::RepeatDisabled
 	);
+
+	CommandList->MapAction(
+		CommandInfos.SelectAll,
+		FExecuteAction::CreateLambda([this]() {
+
+			UContextObjectStore* ContextStore = GetInteractiveToolsContext()->ToolManager->GetContextObjectStore();
+			UUVToolViewportButtonsAPI* UVToolViewportButtonsAPI = ContextStore->FindContext<UUVToolViewportButtonsAPI>();
+			if (UVToolViewportButtonsAPI && UVToolViewportButtonsAPI->AreSelectionButtonsEnabled())
+			{
+				UUVToolSelectionAPI::EUVEditorSelectionMode CurrentSelectionMode = UVToolViewportButtonsAPI->GetSelectionMode();
+				FUVToolSelection::EType SelectionType = FUVToolSelection::EType::Triangle;
+				if (CurrentSelectionMode == UUVToolSelectionAPI::EUVEditorSelectionMode::None)
+				{
+					return; // If we're in none selection mode, don't do anything for select all behavior
+				}
+
+				switch (CurrentSelectionMode)
+				{
+				case UUVToolSelectionAPI::EUVEditorSelectionMode::Vertex:
+					SelectionType = FUVToolSelection::EType::Vertex;
+					break;
+				case UUVToolSelectionAPI::EUVEditorSelectionMode::Edge:
+					SelectionType = FUVToolSelection::EType::Edge;
+					break;
+				case UUVToolSelectionAPI::EUVEditorSelectionMode::Triangle:
+				case UUVToolSelectionAPI::EUVEditorSelectionMode::Island:
+				case UUVToolSelectionAPI::EUVEditorSelectionMode::Mesh:
+					SelectionType = FUVToolSelection::EType::Triangle;
+					break;				
+				default:
+					ensure(false);
+				}
+
+				SelectionAPI->BeginChange();
+				const TArray<FUVToolSelection>& CurrentSelections = SelectionAPI->GetSelections();
+				bool bIsAllSelected = SelectionAPI->HaveSelections() && CurrentSelections.Num() == ToolInputObjects.Num();
+				for (const FUVToolSelection& Selection : CurrentSelections)
+				{
+					bIsAllSelected = bIsAllSelected && Selection.IsAllSelected(*Selection.Target->UnwrapCanonical);
+				}
+
+				SelectionAPI->ClearSelections(false, false);
+				SelectionAPI->ClearUnsetElementAppliedMeshSelections(false, false);
+
+				if (!bIsAllSelected) // If we already had everything selected, unselect everything. Otherwise, select everything.
+				{
+					TArray<FUVToolSelection> AllSelections;
+					AllSelections.SetNum(ToolInputObjects.Num());
+					for (int32 AssetID = 0; AssetID < ToolInputObjects.Num(); ++AssetID)
+					{
+						AllSelections[AssetID].Target = ToolInputObjects[AssetID];
+						AllSelections[AssetID].SelectAll(*ToolInputObjects[AssetID]->UnwrapCanonical, SelectionType);
+					}
+					SelectionAPI->SetSelections(AllSelections, false, false);
+				}
+				SelectionAPI->EndChangeAndEmitIfModified(true);
+			}
+		}),
+		FCanExecuteAction::CreateLambda([this]() { 
+			UContextObjectStore* ContextStore = GetInteractiveToolsContext()->ToolManager->GetContextObjectStore();
+			UUVToolViewportButtonsAPI* UVToolViewportButtonsAPI = ContextStore->FindContext<UUVToolViewportButtonsAPI>();
+			return UVToolViewportButtonsAPI && UVToolViewportButtonsAPI->AreSelectionButtonsEnabled();
+		}));
+
 }
 
 void UUVEditorMode::Exit()
diff --git a/Engine/Plugins/Experimental/UVEditor/Source/UVEditor/Public/UVEditorCommands.h b/Engine/Plugins/Experimental/UVEditor/Source/UVEditor/Public/UVEditorCommands.h
--- a/Engine/Plugins/Experimental/UVEditor/Source/UVEditor/Public/UVEditorCommands.h
+++ b/Engine/Plugins/Experimental/UVEditor/Source/UVEditor/Public/UVEditorCommands.h
@@ -38,6 +38,7 @@
 	TSharedPtr<FUICommandInfo> TriangleSelection;
 	TSharedPtr<FUICommandInfo> IslandSelection;
 	TSharedPtr<FUICommandInfo> FullMeshSelection;
+	TSharedPtr<FUICommandInfo> SelectAll;
 
 	TSharedPtr<FUICommandInfo> EnableOrbitCamera;
 	TSharedPtr<FUICommandInfo> EnableFlyCamera;
diff --git a/Engine/Plugins/Experimental/UVEditor/Source/UVEditorTools/Private/Selection/UVToolSelection.cpp b/Engine/Plugins/Experimental/UVEditor/Source/UVEditorTools/Private/Selection/UVToolSelection.cpp
--- a/Engine/Plugins/Experimental/UVEditor/Source/UVEditorTools/Private/Selection/UVToolSelection.cpp
+++ b/Engine/Plugins/Experimental/UVEditor/Source/UVEditorTools/Private/Selection/UVToolSelection.cpp
@@ -22,7 +22,7 @@ void FUVToolSelection::RestoreFromStableEdgeIdentifiers(const FDynamicMesh3& Mes
 	}
 }
 
-bool FUVToolSelection::AreElementsPresentInMesh(FDynamicMesh3& Mesh) const
+bool FUVToolSelection::AreElementsPresentInMesh(const FDynamicMesh3& Mesh) const
 {
 	switch (Type)
 	{
@@ -57,6 +57,68 @@ bool FUVToolSelection::AreElementsPresentInMesh(FDynamicMesh3& Mesh) const
 	return true;
 }
 
+void FUVToolSelection::SelectAll(const FDynamicMesh3& Mesh, EType TypeIn)
+{
+	Type = TypeIn;
+	StableEdgeIDs.Reset();
+	SelectedIDs.Reset();
+	switch (Type)
+	{
+	case EType::Vertex:
+		SelectedIDs.Reserve(Mesh.VertexCount());
+		for (int32 Vid : Mesh.VertexIndicesItr())
+		{
+			SelectedIDs.Add(Vid);
+		}
+		break;
+	case EType::Edge:
+		SelectedIDs.Reserve(Mesh.EdgeCount());
+		for (int32 Eid : Mesh.EdgeIndicesItr())
+		{
+			SelectedIDs.Add(Eid);
+		}
+		break;
+	case EType::Triangle:
+		SelectedIDs.Reserve(Mesh.TriangleCount());
+		for (int32 Tid : Mesh.TriangleIndicesItr())
+		{
+			SelectedIDs.Add(Tid);
+		}
+		break;
+	default:
+		ensure(false);
+	}
+}
+
+bool FUVToolSelection::IsAllSelected(const FDynamicMesh3& Mesh) const
+{
+	switch (Type)
+	{
+	case EType::Vertex:
+		if (SelectedIDs.Num() != Mesh.VertexCount())
+		{
+			return false;
+		}
+		break;
+	case EType::Edge:
+		if (SelectedIDs.Num() != Mesh.EdgeCount())
+		{
+			return false;
+		}
+		break;
+	case EType::Triangle:
+		if (SelectedIDs.Num() != Mesh.TriangleCount())
+		{
+			return false;
+		}
+		break;
+	default:
+		ensure(false);
+	}
+
+	return AreElementsPresentInMesh(Mesh);
+}
+
 FUVToolSelection FUVToolSelection::GetConvertedSelection(const FDynamicMesh3& Mesh, FUVToolSelection::EType ExpectedSelectionType) const
 {
 	if (Type == ExpectedSelectionType)
diff --git a/Engine/Plugins/Experimental/UVEditor/Source/UVEditorTools/Public/Selection/UVToolSelection.h b/Engine/Plugins/Experimental/UVEditor/Source/UVEditorTools/Public/Selection/UVToolSelection.h
--- a/Engine/Plugins/Experimental/UVEditor/Source/UVEditorTools/Public/Selection/UVToolSelection.h
+++ b/Engine/Plugins/Experimental/UVEditor/Source/UVEditorTools/Public/Selection/UVToolSelection.h
@@ -36,6 +36,10 @@ public:
 		StableEdgeIDs.Reset();
 	}
 
+	void SelectAll(const FDynamicMesh3& Mesh, EType TypeIn);
+
+	bool IsAllSelected(const FDynamicMesh3& Mesh) const;
+
 	bool IsEmpty() const
 	{
 		return SelectedIDs.IsEmpty();
@@ -64,7 +68,7 @@ public:
 
 	void RestoreFromStableEdgeIdentifiers(const FDynamicMesh3& Mesh);
 
-	bool AreElementsPresentInMesh(FDynamicMesh3& Mesh) const;
+	bool AreElementsPresentInMesh(const FDynamicMesh3& Mesh) const;
 
 	/**
 	 * Utility method to construct a new selection from an existing reference selection, only with a new geometry type.
