UVEditor: Replace Align/Distribute icons with new versions, fix mis-labeling bug on vertical/horizontal modes.

#rnx

diff --git a/Engine/Plugins/Experimental/UVEditor/Source/UVEditorTools/Private/Operators/UVEditorUVTransformOp.cpp b/Engine/Plugins/Experimental/UVEditor/Source/UVEditorTools/Private/Operators/UVEditorUVTransformOp.cpp
--- a/Engine/Plugins/Experimental/UVEditor/Source/UVEditorTools/Private/Operators/UVEditorUVTransformOp.cpp
+++ b/Engine/Plugins/Experimental/UVEditor/Source/UVEditorTools/Private/Operators/UVEditorUVTransformOp.cpp
@@ -420,9 +420,9 @@ void FUVEditorUVAlignOp::HandleTransformationOp(FProgressCancel* Progress)
 		case EUVEditorAlignDirectionBackend::Right:
 			return FVector2f(PointTo.X - PointFrom.X, 0);
 		case EUVEditorAlignDirectionBackend::CenterVertically:
-			return FVector2f(0, PointTo.Y - PointFrom.Y);
+			return FVector2f(PointTo.X - PointFrom.X, 0);
 		case EUVEditorAlignDirectionBackend::CenterHorizontally:
-			return FVector2f(PointTo.X - PointFrom.X, 0);			
+			return FVector2f(0, PointTo.Y - PointFrom.Y);
 		default:
 			ensure(false);
 			return FVector2f(0, 0);
@@ -659,10 +659,10 @@ void FUVEditorUVDistributeOp::HandleTransformationOp(FProgressCancel* Progress)
 			PerComponentTranslation = ComputeDistributeTranslations(false, EUVEditorAlignDirectionBackend::Right, -1.0f, false);
 			break;
 		case EUVEditorDistributeModeBackend::CentersVertically:
-			PerComponentTranslation = ComputeDistributeTranslations(false, EUVEditorAlignDirectionBackend::CenterVertically, 1.0f, false);
+			PerComponentTranslation = ComputeDistributeTranslations(true, EUVEditorAlignDirectionBackend::CenterVertically, 1.0f, false);
 			break;
 		case EUVEditorDistributeModeBackend::CentersHorizontally:
-			PerComponentTranslation = ComputeDistributeTranslations(true, EUVEditorAlignDirectionBackend::CenterHorizontally, 1.0f, false);
+			PerComponentTranslation = ComputeDistributeTranslations(false, EUVEditorAlignDirectionBackend::CenterHorizontally, 1.0f, false);
 			break;
 		case EUVEditorDistributeModeBackend::HorizontalSpace:
 			PerComponentTranslation = ComputeDistributeTranslations(false, EUVEditorAlignDirectionBackend::Left, 1.0f, true);
diff --git a/Engine/Plugins/Experimental/UVEditor/Source/UVEditorTools/Private/UVEditorStyle.cpp b/Engine/Plugins/Experimental/UVEditor/Source/UVEditorTools/Private/UVEditorStyle.cpp
--- a/Engine/Plugins/Experimental/UVEditor/Source/UVEditorTools/Private/UVEditorStyle.cpp
+++ b/Engine/Plugins/Experimental/UVEditor/Source/UVEditorTools/Private/UVEditorStyle.cpp
@@ -53,23 +53,25 @@ FUVEditorStyle::FUVEditorStyle()
 	Set("UVEditor.FullMeshSelection", new IMAGE_BRUSH_SVG("SelectionMulti", ToolbarIconSize));
 
 	// Distribute
-	Set("UVEditor.DistributeSpaceVertically", new IMAGE_BRUSH_SVG("DistributeModeSpaceVertically", ToolbarIconSize));
-	Set("UVEditor.DistributeSpaceHorizontally", new IMAGE_BRUSH_SVG("DistributeModeSpaceHorizontally", ToolbarIconSize));
-	Set("UVEditor.DistributeTopEdges", new IMAGE_BRUSH_SVG("DistributeModeTopEdges", ToolbarIconSize));
-	Set("UVEditor.DistributeBottomEdges", new IMAGE_BRUSH_SVG("DistributeModeBottomEdges", ToolbarIconSize));
-	Set("UVEditor.DistributeLeftEdges", new IMAGE_BRUSH_SVG("DistributeModeLeftEdges", ToolbarIconSize));
-	Set("UVEditor.DistributeRightEdges", new IMAGE_BRUSH_SVG("DistributeModeRightEdges", ToolbarIconSize));
-	Set("UVEditor.DistributeCentersHorizontally", new IMAGE_BRUSH_SVG("DistributeModeCentersHorizontally", ToolbarIconSize));
-	Set("UVEditor.DistributeCentersVertically", new IMAGE_BRUSH_SVG("DistributeModeCentersVertically", ToolbarIconSize));
-	Set("UVEditor.DistributeRemoveOverlap", new IMAGE_BRUSH_SVG("DistributeModeRemoveOverlap", ToolbarIconSize));
+	Set("UVEditor.DistributeSpaceVertically", new IMAGE_BRUSH_SVG("DistributeVertical", ToolbarIconSize));
+	Set("UVEditor.DistributeSpaceHorizontally", new IMAGE_BRUSH_SVG("DistributeHorizontal", ToolbarIconSize));
+	Set("UVEditor.DistributeTopEdges", new IMAGE_BRUSH_SVG("DistributeEdgesTop", ToolbarIconSize));
+	Set("UVEditor.DistributeBottomEdges", new IMAGE_BRUSH_SVG("DistributeEdgesBottom", ToolbarIconSize));
+	Set("UVEditor.DistributeLeftEdges", new IMAGE_BRUSH_SVG("DistributeEdgesLeft", ToolbarIconSize));
+	Set("UVEditor.DistributeRightEdges", new IMAGE_BRUSH_SVG("DistributeEdgesRight", ToolbarIconSize));
+	Set("UVEditor.DistributeCentersHorizontally", new IMAGE_BRUSH_SVG("DistributeCentersHorizontal", ToolbarIconSize));
+	Set("UVEditor.DistributeCentersVertically", new IMAGE_BRUSH_SVG("DistributeCentersVertical", ToolbarIconSize));
+	Set("UVEditor.DistributeRemoveOverlap", new IMAGE_BRUSH_SVG("RemoveOverlap", ToolbarIconSize));
 
 	// Align
-	Set("UVEditor.AlignDirectionBottomEdges", new IMAGE_BRUSH_SVG("DistributeModeBottomEdges", ToolbarIconSize));
-	Set("UVEditor.AlignDirectionTopEdges", new IMAGE_BRUSH_SVG("DistributeModeTopEdges", ToolbarIconSize));
-	Set("UVEditor.AlignDirectionLeftEdges", new IMAGE_BRUSH_SVG("DistributeModeLeftEdges", ToolbarIconSize));
-	Set("UVEditor.AlignDirectionRightEdges", new IMAGE_BRUSH_SVG("DistributeModeRightEdges", ToolbarIconSize));
-	Set("UVEditor.AlignDirectionCentersHorizontally", new IMAGE_BRUSH_SVG("DistributeModeCentersVertically", ToolbarIconSize));
-	Set("UVEditor.AlignDirectionCentersVertically", new IMAGE_BRUSH_SVG("DistributeModeCentersHorizontally", ToolbarIconSize));
+	Set("UVEditor.AlignDirectionBottomEdges", new IMAGE_BRUSH_SVG("AlignEdgesBottom", ToolbarIconSize));
+	Set("UVEditor.AlignDirectionTopEdges", new IMAGE_BRUSH_SVG("AlignEdgesTop", ToolbarIconSize));
+	Set("UVEditor.AlignDirectionLeftEdges", new IMAGE_BRUSH_SVG("AlignEdgesLeft", ToolbarIconSize));
+	Set("UVEditor.AlignDirectionRightEdges", new IMAGE_BRUSH_SVG("AlignEdgesRight", ToolbarIconSize));
+	Set("UVEditor.AlignDirectionCentersHorizontally", new IMAGE_BRUSH_SVG("AlignCentersHorizontal", ToolbarIconSize));
+	Set("UVEditor.AlignDirectionCentersVertically", new IMAGE_BRUSH_SVG("AlignCentersVertical", ToolbarIconSize));
+
+
 
 	FSlateStyleRegistry::RegisterSlateStyle(*this);
 }
diff --git a/Engine/Plugins/Experimental/UVEditor/Source/UVEditorTools/Public/Operators/UVEditorUVTransformOp.h b/Engine/Plugins/Experimental/UVEditor/Source/UVEditorTools/Public/Operators/UVEditorUVTransformOp.h
--- a/Engine/Plugins/Experimental/UVEditor/Source/UVEditorTools/Public/Operators/UVEditorUVTransformOp.h
+++ b/Engine/Plugins/Experimental/UVEditor/Source/UVEditorTools/Public/Operators/UVEditorUVTransformOp.h
@@ -202,7 +202,7 @@ public:
 	FVector2D ManualAnchor = FVector2D(0, 0);
 
 	/** Controls what side of the island bounding boxes are being aligned */
-	UPROPERTY(EditAnywhere, Category = "Align", meta = (DisplayName = "Alignment Direction"))
+	UPROPERTY(EditAnywhere, Category = "Align", meta = (TransientToolProperty, DisplayName = "Alignment Direction"))
 	EUVEditorAlignDirection AlignDirection = EUVEditorAlignDirection::None;
 
 	/** Controls how alignment considers grouping selected objects with respect to the alignment behavior.*/
