UVEditor: Support live preview camera snapping to selection upon a UI button or hotkey (Alt-F)

diff --git a/Engine/Plugins/Experimental/UVEditor/Source/UVEditor/Private/SUVEditor3DViewport.cpp b/Engine/Plugins/Experimental/UVEditor/Source/UVEditor/Private/SUVEditor3DViewport.cpp
--- a/Engine/Plugins/Experimental/UVEditor/Source/UVEditor/Private/SUVEditor3DViewport.cpp
+++ b/Engine/Plugins/Experimental/UVEditor/Source/UVEditor/Private/SUVEditor3DViewport.cpp
@@ -29,6 +29,15 @@ void SUVEditor3DViewport::BindCommands()
 			}),
 		FCanExecuteAction::CreateLambda([this]() { return true; }),
 		FIsActionChecked::CreateLambda([this]() { return StaticCastSharedPtr<FUVEditor3DViewportClient>(Client)->GetCameraMode() == EUVEditor3DViewportClientCameraMode::Fly; }));
+
+	CommandList->MapAction(
+		CommandInfos.SetFocusCamera,
+		FExecuteAction::CreateLambda([this]() {
+			StaticCastSharedPtr<FUVEditor3DViewportClient>(Client)->FocusCameraOnSelection();
+		}),
+		FCanExecuteAction::CreateLambda([this]() { 
+			return true;
+		}));
 }
 
 TSharedPtr<SWidget> SUVEditor3DViewport::MakeViewportToolbar()
diff --git a/Engine/Plugins/Experimental/UVEditor/Source/UVEditor/Private/SUVEditor3DViewportToolBar.cpp b/Engine/Plugins/Experimental/UVEditor/Source/UVEditor/Private/SUVEditor3DViewportToolBar.cpp
--- a/Engine/Plugins/Experimental/UVEditor/Source/UVEditor/Private/SUVEditor3DViewportToolBar.cpp
+++ b/Engine/Plugins/Experimental/UVEditor/Source/UVEditor/Private/SUVEditor3DViewportToolBar.cpp
@@ -92,6 +92,15 @@ TSharedRef<SWidget> SUVEditor3DViewportToolBar::MakeToolBar(const TSharedPtr<FEx
 			TAttribute<FSlateIcon>(FSlateIcon(FUVEditorStyle::Get().GetStyleSetName(), "UVEditor.FlyCamera")), FlyCameraName);
 
 		ToolbarBuilder.EndBlockGroup();
+
+		ToolbarBuilder.BeginBlockGroup();
+
+		// Focus Camera
+		static FName FocusCameraName = FName(TEXT("FocusCamera"));
+		ToolbarBuilder.AddToolBarButton(FUVEditorCommands::Get().SetFocusCamera, NAME_None, TAttribute<FText>(), TAttribute<FText>(),
+			TAttribute<FSlateIcon>(FSlateIcon(FUVEditorStyle::Get().GetStyleSetName(), "UVEditor.FocusCamera")), FocusCameraName);
+
+		ToolbarBuilder.EndBlockGroup();
 	}
 
 	ToolbarBuilder.EndSection();
diff --git a/Engine/Plugins/Experimental/UVEditor/Source/UVEditor/Private/UVEditor3DViewportClient.cpp b/Engine/Plugins/Experimental/UVEditor/Source/UVEditor/Private/UVEditor3DViewportClient.cpp
--- a/Engine/Plugins/Experimental/UVEditor/Source/UVEditor/Private/UVEditor3DViewportClient.cpp
+++ b/Engine/Plugins/Experimental/UVEditor/Source/UVEditor/Private/UVEditor3DViewportClient.cpp
@@ -10,11 +10,22 @@
 #include "Drawing/MeshDebugDrawing.h"
 #include "FrameTypes.h"
 #include "UVEditorMode.h"
+#include "ContextObjects/UVToolViewportButtonsAPI.h"
 
 FUVEditor3DViewportClient::FUVEditor3DViewportClient(FEditorModeTools* InModeTools,
-	FPreviewScene* InPreviewScene, const TWeakPtr<SEditorViewport>& InEditorViewportWidget)
-	: FEditorViewportClient(InModeTools, InPreviewScene, InEditorViewportWidget)
+	FPreviewScene* InPreviewScene, const TWeakPtr<SEditorViewport>& InEditorViewportWidget,
+	UUVToolViewportButtonsAPI* ViewportButtonsAPIIn)
+	: FEditorViewportClient(InModeTools, InPreviewScene, InEditorViewportWidget), 
+	ViewportButtonsAPI(ViewportButtonsAPIIn)
 {
 	// We want our near clip plane to be quite close so that we can zoom in further.
 	OverrideNearClipPlane(KINDA_SMALL_NUMBER);
 }
+
+void FUVEditor3DViewportClient::FocusCameraOnSelection()
+{
+	if (ViewportButtonsAPI) 
+	{
+		ViewportButtonsAPI->InitiateFocusCameraOnSelection();
+	}
+}
\ No newline at end of file
diff --git a/Engine/Plugins/Experimental/UVEditor/Source/UVEditor/Private/UVEditorCommands.cpp b/Engine/Plugins/Experimental/UVEditor/Source/UVEditor/Private/UVEditorCommands.cpp
--- a/Engine/Plugins/Experimental/UVEditor/Source/UVEditor/Private/UVEditorCommands.cpp
+++ b/Engine/Plugins/Experimental/UVEditor/Source/UVEditor/Private/UVEditorCommands.cpp
@@ -54,6 +54,8 @@ void FUVEditorCommands::RegisterCommands()
 
 	UI_COMMAND(EnableOrbitCamera, "Orbit", "Enable orbit camera", EUserInterfaceActionType::ToggleButton, FInputChord());
 	UI_COMMAND(EnableFlyCamera, "Fly", "Enable fly camera", EUserInterfaceActionType::ToggleButton, FInputChord());
+	UI_COMMAND(SetFocusCamera, "Focus Camera", "Focus camera around the currently selected UVs", EUserInterfaceActionType::Button, FInputChord(EKeys::F, EModifierKey::Alt));
+
 	UI_COMMAND(ToggleBackground, "Toggle Background", "Toggle background display", EUserInterfaceActionType::ToggleButton,
 	           FInputChord(EModifierKey::Alt, EKeys::B));
 }
diff --git a/Engine/Plugins/Experimental/UVEditor/Source/UVEditor/Private/UVEditorMode.cpp b/Engine/Plugins/Experimental/UVEditor/Source/UVEditor/Private/UVEditorMode.cpp
--- a/Engine/Plugins/Experimental/UVEditor/Source/UVEditor/Private/UVEditorMode.cpp
+++ b/Engine/Plugins/Experimental/UVEditor/Source/UVEditor/Private/UVEditorMode.cpp
@@ -736,6 +736,12 @@ void UUVEditorMode::BindCommands()
 			return UVToolViewportButtonsAPI && UVToolViewportButtonsAPI->AreSelectionButtonsEnabled();
 		}));
 
+	UContextObjectStore* ContextStore = GetInteractiveToolsContext()->ToolManager->GetContextObjectStore();
+	UUVToolViewportButtonsAPI* UVToolViewportButtonsAPI = ContextStore->FindContext<UUVToolViewportButtonsAPI>();
+	if (UVToolViewportButtonsAPI)
+	{
+		UVToolViewportButtonsAPI->OnInitiateFocusCameraOnSelection.AddUObject(this, &UUVEditorMode::FocusLivePreviewCameraOnSelection);
+	}
 }
 
 void UUVEditorMode::Exit()
@@ -834,7 +840,11 @@ void UUVEditorMode::InitializeAssetEditorContexts(UContextObjectStore& ContextSt
 			LivePreviewModeManager.GetInteractiveToolsContext()->InputRouter,
 			[LivePreviewViewportClientPtr = &LivePreviewViewportClient](FViewCameraState& CameraStateOut) {
 				GetCameraState(*LivePreviewViewportClientPtr, CameraStateOut);
-			});
+			},
+			[LivePreviewViewportClientPtr = &LivePreviewViewportClient](const FAxisAlignedBox3d& BoundingBox) {
+				LivePreviewViewportClientPtr->FocusViewportOnBox((FBox)BoundingBox, true);
+			}
+			);
 		ContextStore.AddContextObject(LivePreviewAPI);
 	}
 
@@ -1195,6 +1205,39 @@ void UUVEditorMode::UpdatePreviewMaterialBasedOnBackground()
 	}
 }
 
+void UUVEditorMode::FocusLivePreviewCameraOnSelection()
+{
+	UContextObjectStore* ContextStore = GetInteractiveToolsContext()->ToolManager->GetContextObjectStore();
+	UUVToolLivePreviewAPI* LivePreviewAPI = ContextStore->FindContext<UUVToolLivePreviewAPI>();
+	if (!LivePreviewAPI)
+	{
+		return;
+	}
+
+	FAxisAlignedBox3d SelectionBoundingBox;
+
+	const TArray<FUVToolSelection>& CurrentSelections = SelectionAPI->GetSelections();
+	const TArray<FUVToolSelection>& CurrentUnsetSelections = SelectionAPI->GetUnsetElementAppliedMeshSelections();
+
+	for (const FUVToolSelection& Selection : CurrentSelections)
+	{
+		SelectionBoundingBox.Contain(Selection.ToBoundingBox(*Selection.Target->AppliedCanonical));
+	}
+	for (const FUVToolSelection& Selection : CurrentUnsetSelections)
+	{
+		SelectionBoundingBox.Contain(Selection.ToBoundingBox(*Selection.Target->AppliedCanonical));
+	}
+	if (CurrentSelections.Num() == 0 && CurrentUnsetSelections.Num() == 0)
+	{
+		for (int32 AssetID = 0; AssetID < ToolInputObjects.Num(); ++AssetID)
+		{
+			SelectionBoundingBox.Contain(ToolInputObjects[AssetID]->AppliedCanonical->GetBounds());
+		}
+	}
+	
+	LivePreviewAPI->SetLivePreviewCameraToLookAtVolume(SelectionBoundingBox);
+}
+
 void UUVEditorMode::ModeTick(float DeltaTime)
 {
 	using namespace UVEditorModeLocals;
diff --git a/Engine/Plugins/Experimental/UVEditor/Source/UVEditor/Private/UVEditorToolkit.cpp b/Engine/Plugins/Experimental/UVEditor/Source/UVEditor/Private/UVEditorToolkit.cpp
--- a/Engine/Plugins/Experimental/UVEditor/Source/UVEditor/Private/UVEditorToolkit.cpp
+++ b/Engine/Plugins/Experimental/UVEditor/Source/UVEditor/Private/UVEditorToolkit.cpp
@@ -93,6 +93,13 @@ FUVEditorToolkit::FUVEditorToolkit(UAssetEditor* InOwningAssetEditor)
 	UVEditorModule->OnRegisterLayoutExtensions().Broadcast(*LayoutExtender);
 	StandaloneDefaultLayout->ProcessExtensions(*LayoutExtender);
 
+	// This API object serves as a communication point between the viewport toolbars and the tools.
+    // We create it here so that we can pass it both into the 2d & 3d viewports and when we initialize 
+    // the mode.
+	ViewportButtonsAPI = NewObject<UUVToolViewportButtonsAPI>();
+
+	UVTool2DViewportAPI = NewObject<UUVTool2DViewportAPI>();
+
 	// We could create the preview scenes in CreateEditorViewportClient() the way that FBaseAssetToolkit
 	// does, but it seems more intuitive to create them right off the bat and pass it in later. 
 	FPreviewScene::ConstructionValues PreviewSceneArgs;
@@ -106,7 +113,7 @@ FUVEditorToolkit::FUVEditorToolkit(UAssetEditor* InOwningAssetEditor)
 
 	LivePreviewTabContent = MakeShareable(new FEditorViewportTabContent());
 	LivePreviewViewportClient = MakeShared<FUVEditor3DViewportClient>(
-		LivePreviewEditorModeManager.Get(), LivePreviewScene.Get());
+		LivePreviewEditorModeManager.Get(), LivePreviewScene.Get(), nullptr, ViewportButtonsAPI);
 
 	LivePreviewViewportDelegate = [this](FAssetEditorViewportConstructionArgs InArgs)
 	{
@@ -114,12 +121,7 @@ FUVEditorToolkit::FUVEditorToolkit(UAssetEditor* InOwningAssetEditor)
 			.EditorViewportClient(LivePreviewViewportClient);
 	};
 
-	// This API object serves as a communication point between the viewport toolbars and the tools.
-	// We create it here so that we can pass it both into the 2d viewport and when we initialize 
-	// the mode.
-	ViewportButtonsAPI = NewObject<UUVToolViewportButtonsAPI>();
 
-	UVTool2DViewportAPI = NewObject<UUVTool2DViewportAPI>();
 }
 
 FUVEditorToolkit::~FUVEditorToolkit()
diff --git a/Engine/Plugins/Experimental/UVEditor/Source/UVEditor/Public/UVEditor3DViewportClient.h b/Engine/Plugins/Experimental/UVEditor/Source/UVEditor/Public/UVEditor3DViewportClient.h
--- a/Engine/Plugins/Experimental/UVEditor/Source/UVEditor/Public/UVEditor3DViewportClient.h
+++ b/Engine/Plugins/Experimental/UVEditor/Source/UVEditor/Public/UVEditor3DViewportClient.h
@@ -6,6 +6,7 @@
 
 #include "EditorViewportClient.h"
 
+class UUVToolViewportButtonsAPI;
 
 // Types of camera motion for the UV Editor 3D viewport
 enum EUVEditor3DViewportClientCameraMode {
@@ -22,7 +23,7 @@ class UVEDITOR_API FUVEditor3DViewportClient : public FEditorViewportClient
 public:
 
 	FUVEditor3DViewportClient(FEditorModeTools* InModeTools, FPreviewScene* InPreviewScene = nullptr,
-		const TWeakPtr<SEditorViewport>& InEditorViewportWidget = nullptr);
+		const TWeakPtr<SEditorViewport>& InEditorViewportWidget = nullptr, UUVToolViewportButtonsAPI* ViewportButtonsAPI = nullptr);
 
 	virtual ~FUVEditor3DViewportClient() {}
 
@@ -50,7 +51,7 @@ public:
 	bool CanSetWidgetMode(UE::Widget::EWidgetMode NewMode) const override {	return false; }
 	void SetWidgetMode(UE::Widget::EWidgetMode NewMode) override {}
 	UE::Widget::EWidgetMode GetWidgetMode() const override { return UE::Widget::EWidgetMode::WM_None; }
-
+	void FocusCameraOnSelection();
 public:
 
 	void SetCameraMode(EUVEditor3DViewportClientCameraMode CameraModeIn) { CameraMode = CameraModeIn; };
@@ -61,4 +62,6 @@ protected:
 	// Enforce Orbit camera for UV editor live preview viewport. Use this instead of the base class orbit camera flag
 	// to allow for expected behaviors of the base class when in fly camera mode.
 	EUVEditor3DViewportClientCameraMode CameraMode = EUVEditor3DViewportClientCameraMode::Orbit;
+	UUVToolViewportButtonsAPI* ViewportButtonsAPI;
+
 };
\ No newline at end of file
diff --git a/Engine/Plugins/Experimental/UVEditor/Source/UVEditor/Public/UVEditorCommands.h b/Engine/Plugins/Experimental/UVEditor/Source/UVEditor/Public/UVEditorCommands.h
--- a/Engine/Plugins/Experimental/UVEditor/Source/UVEditor/Public/UVEditorCommands.h
+++ b/Engine/Plugins/Experimental/UVEditor/Source/UVEditor/Public/UVEditorCommands.h
@@ -42,5 +42,7 @@ public:
 
 	TSharedPtr<FUICommandInfo> EnableOrbitCamera;
 	TSharedPtr<FUICommandInfo> EnableFlyCamera;
+	TSharedPtr<FUICommandInfo> SetFocusCamera;
+
 	TSharedPtr<FUICommandInfo> ToggleBackground;
 };
\ No newline at end of file
diff --git a/Engine/Plugins/Experimental/UVEditor/Source/UVEditor/Public/UVEditorMode.h b/Engine/Plugins/Experimental/UVEditor/Source/UVEditor/Public/UVEditorMode.h
--- a/Engine/Plugins/Experimental/UVEditor/Source/UVEditor/Public/UVEditorMode.h
+++ b/Engine/Plugins/Experimental/UVEditor/Source/UVEditor/Public/UVEditorMode.h
@@ -262,6 +262,7 @@ protected:
 	
 	void UpdateTriangleMaterialBasedOnBackground(bool IsBackgroundVisible);
 	void UpdatePreviewMaterialBasedOnBackground();
+	void FocusLivePreviewCameraOnSelection();
 
 	void UpdateActiveUDIMs();
 	int32 UDIMsChangedWatcherId;
diff --git a/Engine/Plugins/Experimental/UVEditor/Source/UVEditorTools/Private/ContextObjects/UVToolContextObjects.cpp b/Engine/Plugins/Experimental/UVEditor/Source/UVEditorTools/Private/ContextObjects/UVToolContextObjects.cpp
--- a/Engine/Plugins/Experimental/UVEditor/Source/UVEditorTools/Private/ContextObjects/UVToolContextObjects.cpp
+++ b/Engine/Plugins/Experimental/UVEditor/Source/UVEditorTools/Private/ContextObjects/UVToolContextObjects.cpp
@@ -77,11 +77,13 @@ void UUVToolEmitChangeAPI::EmitToolDependentChange(UObject* TargetObject, TUniqu
 }
 
 void UUVToolLivePreviewAPI::Initialize(UWorld* WorldIn, UInputRouter* RouterIn,
-	TUniqueFunction<void(FViewCameraState& CameraStateOut)> GetLivePreviewCameraStateFuncIn)
+	TUniqueFunction<void(FViewCameraState& CameraStateOut)> GetLivePreviewCameraStateFuncIn,
+	TUniqueFunction<void(const FAxisAlignedBox3d& BoundingBox)> SetLivePreviewCameraToLookAtVolumeFuncIn)
 {
 	World = WorldIn;
 	InputRouter = RouterIn;
 	GetLivePreviewCameraStateFunc = MoveTemp(GetLivePreviewCameraStateFuncIn);
+	SetLivePreviewCameraToLookAtVolumeFunc = MoveTemp(SetLivePreviewCameraToLookAtVolumeFuncIn);
 }
 
 void UUVToolLivePreviewAPI::OnToolEnded(UInteractiveTool* DeadTool)
diff --git a/Engine/Plugins/Experimental/UVEditor/Source/UVEditorTools/Private/Selection/UVToolSelection.cpp b/Engine/Plugins/Experimental/UVEditor/Source/UVEditorTools/Private/Selection/UVToolSelection.cpp
--- a/Engine/Plugins/Experimental/UVEditor/Source/UVEditorTools/Private/Selection/UVToolSelection.cpp
+++ b/Engine/Plugins/Experimental/UVEditor/Source/UVEditorTools/Private/Selection/UVToolSelection.cpp
@@ -119,6 +119,46 @@ bool FUVToolSelection::IsAllSelected(const FDynamicMesh3& Mesh) const
 	return AreElementsPresentInMesh(Mesh);
 }
 
+FAxisAlignedBox3d FUVToolSelection::ToBoundingBox(const FDynamicMesh3& Mesh) const
+{
+	FAxisAlignedBox3d BoundingBox;
+
+	switch (Type)
+	{
+	case EType::Vertex:
+		for (int32 Vid : SelectedIDs)
+		{
+			if (Mesh.IsVertex(Vid))
+			{
+				BoundingBox.Contain(Mesh.GetVertexRef(Vid));
+			}
+		}
+		break;
+	case EType::Edge:
+		for (int32 Eid : SelectedIDs)
+		{
+			if (Mesh.IsEdge(Eid))
+			{
+				BoundingBox.Contain(Mesh.GetVertexRef(Mesh.GetEdgeRef(Eid).Vert.A));
+				BoundingBox.Contain(Mesh.GetVertexRef(Mesh.GetEdgeRef(Eid).Vert.B));
+			}
+		}
+		break;
+	case EType::Triangle:
+		for (int32 Tid : SelectedIDs)
+		{
+			if (Mesh.IsTriangle(Tid))
+			{
+				BoundingBox.Contain(Mesh.GetVertexRef(Mesh.GetTriangleRef(Tid).A));
+				BoundingBox.Contain(Mesh.GetVertexRef(Mesh.GetTriangleRef(Tid).B));
+				BoundingBox.Contain(Mesh.GetVertexRef(Mesh.GetTriangleRef(Tid).C));
+			}
+		}
+		break;
+	}
+	return BoundingBox;
+}
+
 FUVToolSelection FUVToolSelection::GetConvertedSelection(const FDynamicMesh3& Mesh, FUVToolSelection::EType ExpectedSelectionType) const
 {
 	if (Type == ExpectedSelectionType)
diff --git a/Engine/Plugins/Experimental/UVEditor/Source/UVEditorTools/Private/UVEditorStyle.cpp b/Engine/Plugins/Experimental/UVEditor/Source/UVEditorTools/Private/UVEditorStyle.cpp
--- a/Engine/Plugins/Experimental/UVEditor/Source/UVEditorTools/Private/UVEditorStyle.cpp
+++ b/Engine/Plugins/Experimental/UVEditor/Source/UVEditorTools/Private/UVEditorStyle.cpp
@@ -45,6 +45,7 @@ FUVEditorStyle::FUVEditorStyle()
 	// Viewport icons
 	Set("UVEditor.OrbitCamera", new CORE_IMAGE_BRUSH_SVG("Starship/EditorViewport/rotate", ToolbarIconSize));
 	Set("UVEditor.FlyCamera", new CORE_IMAGE_BRUSH_SVG("Starship/EditorViewport/camera", ToolbarIconSize));
+	Set("UVEditor.FocusCamera", new CORE_IMAGE_BRUSH_SVG("Starship/Actors/snap-view-to-object", ToolbarIconSize));
 
 	Set("UVEditor.VertexSelection", new IMAGE_BRUSH_SVG("SelectionVertices", ToolbarIconSize));
 	Set("UVEditor.EdgeSelection", new IMAGE_BRUSH_SVG("SelectionLine", ToolbarIconSize));
diff --git a/Engine/Plugins/Experimental/UVEditor/Source/UVEditorTools/Public/ContextObjects/UVToolContextObjects.h b/Engine/Plugins/Experimental/UVEditor/Source/UVEditorTools/Public/ContextObjects/UVToolContextObjects.h
--- a/Engine/Plugins/Experimental/UVEditor/Source/UVEditorTools/Public/ContextObjects/UVToolContextObjects.h
+++ b/Engine/Plugins/Experimental/UVEditor/Source/UVEditorTools/Public/ContextObjects/UVToolContextObjects.h
@@ -107,7 +107,8 @@ class UVEDITORTOOLS_API UUVToolLivePreviewAPI : public UUVToolContextObject
 public:
 
 	void Initialize(UWorld* WorldIn, UInputRouter* RouterIn,
-		TUniqueFunction<void(FViewCameraState& CameraStateOut)> GetLivePreviewCameraStateFuncIn);
+		TUniqueFunction<void(FViewCameraState& CameraStateOut)> GetLivePreviewCameraStateFuncIn,
+		TUniqueFunction<void(const UE::Geometry::FAxisAlignedBox3d& BoundingBox)> SetLivePreviewCameraToLookAtVolumeFuncIn);
 
 	UWorld* GetLivePreviewWorld() { return World.Get(); }
 	UInputRouter* GetLivePreviewInputRouter() { return InputRouter.Get(); }
@@ -119,6 +120,14 @@ public:
 		}
 	}
 
+	void SetLivePreviewCameraToLookAtVolume(const UE::Geometry::FAxisAlignedBox3d& BoundingBox)
+	{
+		if (SetLivePreviewCameraToLookAtVolumeFunc)
+		{
+			SetLivePreviewCameraToLookAtVolumeFunc(BoundingBox);
+		}
+	}
+
 	virtual void OnToolEnded(UInteractiveTool* DeadTool) override;
 
 	/**
@@ -143,6 +152,7 @@ protected:
 	TWeakObjectPtr<UInputRouter> InputRouter;
 
 	TUniqueFunction<void(FViewCameraState& CameraStateOut)> GetLivePreviewCameraStateFunc;
+	TUniqueFunction<void(const UE::Geometry::FAxisAlignedBox3d& BoundingBox)> SetLivePreviewCameraToLookAtVolumeFunc;
 };
 
 USTRUCT()
diff --git a/Engine/Plugins/Experimental/UVEditor/Source/UVEditorTools/Public/ContextObjects/UVToolViewportButtonsAPI.h b/Engine/Plugins/Experimental/UVEditor/Source/UVEditorTools/Public/ContextObjects/UVToolViewportButtonsAPI.h
--- a/Engine/Plugins/Experimental/UVEditor/Source/UVEditorTools/Public/ContextObjects/UVToolViewportButtonsAPI.h
+++ b/Engine/Plugins/Experimental/UVEditor/Source/UVEditorTools/Public/ContextObjects/UVToolViewportButtonsAPI.h
@@ -74,12 +74,20 @@ public:
 		return SelectionMode;
 	}
 
+	void InitiateFocusCameraOnSelection() const
+	{
+		OnInitiateFocusCameraOnSelection.Broadcast();
+	}
+
 	DECLARE_MULTICAST_DELEGATE_OneParam(FOnGizmoModeChange, EGizmoMode NewGizmoMode);
 	FOnGizmoModeChange OnGizmoModeChange;
 
 	DECLARE_MULTICAST_DELEGATE_OneParam(FOnSelectionModeChange, ESelectionMode NewSelectionMode);
 	FOnSelectionModeChange OnSelectionModeChange;
 
+	DECLARE_MULTICAST_DELEGATE(FOnInitiateFocusCameraOnSelection);
+	FOnInitiateFocusCameraOnSelection OnInitiateFocusCameraOnSelection;
+
 	virtual void OnToolEnded(UInteractiveTool* DeadTool) override 
 	{
 		OnGizmoModeChange.RemoveAll(DeadTool);
diff --git a/Engine/Plugins/Experimental/UVEditor/Source/UVEditorTools/Public/Selection/UVToolSelection.h b/Engine/Plugins/Experimental/UVEditor/Source/UVEditorTools/Public/Selection/UVToolSelection.h
--- a/Engine/Plugins/Experimental/UVEditor/Source/UVEditorTools/Public/Selection/UVToolSelection.h
+++ b/Engine/Plugins/Experimental/UVEditor/Source/UVEditorTools/Public/Selection/UVToolSelection.h
@@ -70,6 +70,7 @@ public:
 
 	bool AreElementsPresentInMesh(const FDynamicMesh3& Mesh) const;
 
+	FAxisAlignedBox3d ToBoundingBox(const FDynamicMesh3& Mesh) const;
 	/**
 	 * Utility method to construct a new selection from an existing reference selection, only with a new geometry type.
 	 */
