UVEditor: Fix bug where 2D viewport allowed orbiting due to usage of deprecated API.

diff --git a/Engine/Plugins/Experimental/UVEditor/Source/UVEditor/Private/UVEditor2DViewportClient.cpp b/Engine/Plugins/Experimental/UVEditor/Source/UVEditor/Private/UVEditor2DViewportClient.cpp
--- a/Engine/Plugins/Experimental/UVEditor/Source/UVEditor/Private/UVEditor2DViewportClient.cpp
+++ b/Engine/Plugins/Experimental/UVEditor/Source/UVEditor/Private/UVEditor2DViewportClient.cpp
@@ -213,7 +213,7 @@ void FUVEditor2DViewportClient::AddReferencedObjects(FReferenceCollector& Collec
 	Collector.AddReferencedObject(CanvasObject);
 }
 
-bool FUVEditor2DViewportClient::InputKey(FViewport* InViewport, int32 ControllerId, FKey Key, EInputEvent Event, float/*AmountDepressed*/, bool/*Gamepad*/)
+bool FUVEditor2DViewportClient::InputKey(const FInputKeyEventArgs& EventArgs)
 {
 	// We'll support disabling input like our base class, even if it does not end up being used.
 	if (bDisableInput)
@@ -222,9 +222,11 @@ bool FUVEditor2DViewportClient::InputKey(FViewport* InViewport, int32 Controller
 	}
 
 	// Our viewport manipulation is placed in the input router that ModeTools manages
-	return ModeTools->InputKey(this, Viewport, Key, Event);
+	return ModeTools->InputKey(this, EventArgs.Viewport, EventArgs.Key, EventArgs.Event);
+
 }
 
+
 // Note that this function gets called from the super class Draw(FViewport*, FCanvas*) overload to draw the scene.
 // We don't override that top-level function so that it can do whatever view calculations it needs to do.
 void FUVEditor2DViewportClient::Draw(const FSceneView* View, FPrimitiveDrawInterface* PDI)
diff --git a/Engine/Plugins/Experimental/UVEditor/Source/UVEditor/Public/UVEditor2DViewportClient.h b/Engine/Plugins/Experimental/UVEditor/Source/UVEditor/Public/UVEditor2DViewportClient.h
--- a/Engine/Plugins/Experimental/UVEditor/Source/UVEditor/Public/UVEditor2DViewportClient.h
+++ b/Engine/Plugins/Experimental/UVEditor/Source/UVEditor/Public/UVEditor2DViewportClient.h
@@ -30,7 +30,8 @@ public:
 	bool AreWidgetButtonsEnabled() const;
 
 	// FEditorViewportClient
-	virtual bool InputKey(FViewport* InViewport, int32 ControllerId, FKey Key, EInputEvent Event, float/*AmountDepressed*/, bool/*Gamepad*/) override;
+	virtual bool InputKey(const FInputKeyEventArgs& EventArgs) override;
+
 	virtual void Draw(const FSceneView* View, FPrimitiveDrawInterface* PDI) override;
 	virtual bool ShouldOrbitCamera() const override;
 	bool CanSetWidgetMode(UE::Widget::EWidgetMode NewMode) const override;
