UVEditor: Focus camera on loaded assets upon loading asset editor. Makes use of existing focus behavior, just calls it automatically at start up.

diff --git a/Engine/Plugins/Experimental/UVEditor/Source/UVEditor/Private/UVEditorToolkit.cpp b/Engine/Plugins/Experimental/UVEditor/Source/UVEditor/Private/UVEditorToolkit.cpp
--- a/Engine/Plugins/Experimental/UVEditor/Source/UVEditor/Private/UVEditorToolkit.cpp
+++ b/Engine/Plugins/Experimental/UVEditor/Source/UVEditor/Private/UVEditorToolkit.cpp
@@ -564,6 +564,10 @@ void FUVEditorToolkit::PostInitAssetEditor()
 	LivePreviewViewportClient->SetViewLocation(FVector(-200, 100, 100));
 	LivePreviewViewportClient->SetLookAtLocation(FVector(0, 0, 0));
 
+	// Adjust camera view to focus on the scene
+	UVMode->FocusLivePreviewCameraOnSelection();
+
+
 	// Hook up the viewport command list to our toolkit command list so that hotkeys not
 	// handled by our toolkit would be handled by the viewport (to allow us to use
 	// whatever hotkeys the viewport registered after clicking in the detail panel or
