UVEditor: Adjust handling of the select all / deselect functionality to better match the behavior in the Level Editor. Moves selection logic to the SelectionTool instead of the UVEditorMode.

diff --git a/Engine/Plugins/Experimental/UVEditor/Source/UVEditor/Private/UVEditorCommands.cpp b/Engine/Plugins/Experimental/UVEditor/Source/UVEditor/Private/UVEditorCommands.cpp
--- a/Engine/Plugins/Experimental/UVEditor/Source/UVEditor/Private/UVEditorCommands.cpp
+++ b/Engine/Plugins/Experimental/UVEditor/Source/UVEditor/Private/UVEditorCommands.cpp
@@ -42,7 +42,7 @@ void FUVEditorCommands::RegisterCommands()
 
 	// These allow us to link up to pressed keys
 	UI_COMMAND(AcceptOrCompleteActiveTool, "Accept", "Accept the active tool", EUserInterfaceActionType::Button, FInputChord(EKeys::Enter));
-	UI_COMMAND(CancelOrCompleteActiveTool, "Cancel", "Cancel the active tool", EUserInterfaceActionType::Button, FInputChord(EKeys::Escape));
+	UI_COMMAND(CancelOrCompleteActiveTool, "Cancel", "Cancel the active tool or clear current selection", EUserInterfaceActionType::Button, FInputChord(EKeys::Escape));
 
 	// These get used in viewport buttons
 	UI_COMMAND(VertexSelection, "Vertex Selection", "Select vertices", EUserInterfaceActionType::ToggleButton, FInputChord(EKeys::One));
@@ -50,7 +50,7 @@ void FUVEditorCommands::RegisterCommands()
 	UI_COMMAND(TriangleSelection, "Triangle Selection", "Select triangles", EUserInterfaceActionType::ToggleButton, FInputChord(EKeys::Three));
 	UI_COMMAND(IslandSelection, "Island Selection", "Select islands", EUserInterfaceActionType::ToggleButton, FInputChord(EKeys::Four));
 	UI_COMMAND(FullMeshSelection, "Mesh Selection", "Select meshes", EUserInterfaceActionType::ToggleButton, FInputChord(EKeys::Five));
-	UI_COMMAND(SelectAll, "Select All", "Toggle selecting everything based on current selection mode", EUserInterfaceActionType::None, FInputChord(EKeys::A, EModifierKey::Control));
+	UI_COMMAND(SelectAll, "Select All", "Select everything based on current selection mode", EUserInterfaceActionType::None, FInputChord(EKeys::A, EModifierKey::Control));
 
 	UI_COMMAND(EnableOrbitCamera, "Orbit", "Enable orbit camera", EUserInterfaceActionType::ToggleButton, FInputChord());
 	UI_COMMAND(EnableFlyCamera, "Fly", "Enable fly camera", EUserInterfaceActionType::ToggleButton, FInputChord());
diff --git a/Engine/Plugins/Experimental/UVEditor/Source/UVEditor/Private/UVEditorMode.cpp b/Engine/Plugins/Experimental/UVEditor/Source/UVEditor/Private/UVEditorMode.cpp
--- a/Engine/Plugins/Experimental/UVEditor/Source/UVEditor/Private/UVEditorMode.cpp
+++ b/Engine/Plugins/Experimental/UVEditor/Source/UVEditor/Private/UVEditorMode.cpp
@@ -677,63 +677,16 @@ void UUVEditorMode::BindCommands()
 		CommandInfos.SelectAll,
 		FExecuteAction::CreateLambda([this]() {
 
-			UContextObjectStore* ContextStore = GetInteractiveToolsContext()->ToolManager->GetContextObjectStore();
-			UUVToolViewportButtonsAPI* UVToolViewportButtonsAPI = ContextStore->FindContext<UUVToolViewportButtonsAPI>();
-			if (UVToolViewportButtonsAPI && UVToolViewportButtonsAPI->AreSelectionButtonsEnabled())
+			if (IsDefaultToolActive())
 			{
-				UUVToolSelectionAPI::EUVEditorSelectionMode CurrentSelectionMode = UVToolViewportButtonsAPI->GetSelectionMode();
-				FUVToolSelection::EType SelectionType = FUVToolSelection::EType::Triangle;
-				if (CurrentSelectionMode == UUVToolSelectionAPI::EUVEditorSelectionMode::None)
-				{
-					return; // If we're in none selection mode, don't do anything for select all behavior
-				}
-
-				switch (CurrentSelectionMode)
-				{
-				case UUVToolSelectionAPI::EUVEditorSelectionMode::Vertex:
-					SelectionType = FUVToolSelection::EType::Vertex;
-					break;
-				case UUVToolSelectionAPI::EUVEditorSelectionMode::Edge:
-					SelectionType = FUVToolSelection::EType::Edge;
-					break;
-				case UUVToolSelectionAPI::EUVEditorSelectionMode::Triangle:
-				case UUVToolSelectionAPI::EUVEditorSelectionMode::Island:
-				case UUVToolSelectionAPI::EUVEditorSelectionMode::Mesh:
-					SelectionType = FUVToolSelection::EType::Triangle;
-					break;				
-				default:
-					ensure(false);
-				}
-
-				SelectionAPI->BeginChange();
-				const TArray<FUVToolSelection>& CurrentSelections = SelectionAPI->GetSelections();
-				bool bIsAllSelected = SelectionAPI->HaveSelections() && CurrentSelections.Num() == ToolInputObjects.Num();
-				for (const FUVToolSelection& Selection : CurrentSelections)
-				{
-					bIsAllSelected = bIsAllSelected && Selection.IsAllSelected(*Selection.Target->UnwrapCanonical);
-				}
-
-				SelectionAPI->ClearSelections(false, false);
-				SelectionAPI->ClearUnsetElementAppliedMeshSelections(false, false);
-
-				if (!bIsAllSelected) // If we already had everything selected, unselect everything. Otherwise, select everything.
-				{
-					TArray<FUVToolSelection> AllSelections;
-					AllSelections.SetNum(ToolInputObjects.Num());
-					for (int32 AssetID = 0; AssetID < ToolInputObjects.Num(); ++AssetID)
-					{
-						AllSelections[AssetID].Target = ToolInputObjects[AssetID];
-						AllSelections[AssetID].SelectAll(*ToolInputObjects[AssetID]->UnwrapCanonical, SelectionType);
-					}
-					SelectionAPI->SetSelections(AllSelections, false, false);
-				}
-				SelectionAPI->EndChangeAndEmitIfModified(true);
+				UInteractiveTool* Tool = GetToolManager()->GetActiveTool(EToolSide::Mouse);
+				UUVSelectTool* SelectTool = Cast<UUVSelectTool>(Tool);
+				check(SelectTool);
+				SelectTool->SelectAll();
 			}
 		}),
 		FCanExecuteAction::CreateLambda([this]() { 
-			UContextObjectStore* ContextStore = GetInteractiveToolsContext()->ToolManager->GetContextObjectStore();
-			UUVToolViewportButtonsAPI* UVToolViewportButtonsAPI = ContextStore->FindContext<UUVToolViewportButtonsAPI>();
-			return UVToolViewportButtonsAPI && UVToolViewportButtonsAPI->AreSelectionButtonsEnabled();
+			return IsDefaultToolActive();
 		}));
 
 	UContextObjectStore* ContextStore = GetInteractiveToolsContext()->ToolManager->GetContextObjectStore();
diff --git a/Engine/Plugins/Experimental/UVEditor/Source/UVEditorTools/Private/UVSelectTool.cpp b/Engine/Plugins/Experimental/UVEditor/Source/UVEditorTools/Private/UVSelectTool.cpp
--- a/Engine/Plugins/Experimental/UVEditor/Source/UVEditorTools/Private/UVSelectTool.cpp
+++ b/Engine/Plugins/Experimental/UVEditor/Source/UVEditorTools/Private/UVSelectTool.cpp
@@ -595,4 +595,66 @@ void UUVSelectTool::OnTick(float DeltaTime)
 	}
 }
 
+bool UUVSelectTool::CanCurrentlyNestedCancel()
+{
+	return CurrentSelections.Num() > 0;
+}
+
+bool UUVSelectTool::ExecuteNestedCancelCommand()
+{
+	if (CurrentSelections.Num() > 0)
+	{
+		SelectionAPI->BeginChange();
+		SelectionAPI->ClearSelections(false, false);
+		SelectionAPI->ClearUnsetElementAppliedMeshSelections(false, false);
+		SelectionAPI->EndChangeAndEmitIfModified(true);
+
+		return true;
+	}
+
+	return false;
+}
+
+void UUVSelectTool::SelectAll()
+{
+	UUVToolSelectionAPI::EUVEditorSelectionMode CurrentSelectionMode = ViewportButtonsAPI->GetSelectionMode();
+	FUVToolSelection::EType SelectionType = FUVToolSelection::EType::Triangle;
+	if (CurrentSelectionMode == UUVToolSelectionAPI::EUVEditorSelectionMode::None)
+	{
+		return; // If we're in none selection mode, don't do anything for select all behavior
+	}
+
+	switch (CurrentSelectionMode)
+	{
+	case UUVToolSelectionAPI::EUVEditorSelectionMode::Vertex:
+		SelectionType = FUVToolSelection::EType::Vertex;
+		break;
+	case UUVToolSelectionAPI::EUVEditorSelectionMode::Edge:
+		SelectionType = FUVToolSelection::EType::Edge;
+		break;
+	case UUVToolSelectionAPI::EUVEditorSelectionMode::Triangle:
+	case UUVToolSelectionAPI::EUVEditorSelectionMode::Island:
+	case UUVToolSelectionAPI::EUVEditorSelectionMode::Mesh:
+		SelectionType = FUVToolSelection::EType::Triangle;
+		break;
+	default:
+		ensure(false);
+	}
+
+	SelectionAPI->BeginChange();
+	SelectionAPI->ClearSelections(false, false);
+	SelectionAPI->ClearUnsetElementAppliedMeshSelections(false, false);
+
+	TArray<FUVToolSelection> AllSelections;
+	AllSelections.SetNum(Targets.Num());
+	for (int32 AssetID = 0; AssetID < Targets.Num(); ++AssetID)
+	{
+		AllSelections[AssetID].Target = Targets[AssetID];
+		AllSelections[AssetID].SelectAll(*Targets[AssetID]->UnwrapCanonical, SelectionType);
+	}
+	SelectionAPI->SetSelections(AllSelections, false, false);
+
+	SelectionAPI->EndChangeAndEmitIfModified(true);
+}
+
 #undef LOCTEXT_NAMESPACE
diff --git a/Engine/Plugins/Experimental/UVEditor/Source/UVEditorTools/Public/UVSelectTool.h b/Engine/Plugins/Experimental/UVEditor/Source/UVEditorTools/Public/UVSelectTool.h
--- a/Engine/Plugins/Experimental/UVEditor/Source/UVEditorTools/Public/UVSelectTool.h
+++ b/Engine/Plugins/Experimental/UVEditor/Source/UVEditorTools/Public/UVSelectTool.h
@@ -10,6 +10,7 @@
 #include "InteractiveToolBuilder.h"
 #include "Selection/UVToolSelection.h"
 #include "Selection/UVToolSelectionAPI.h" // IUVToolSupportsSelection
+#include "InteractiveToolQueryInterfaces.h" // IInteractiveToolNestedAcceptCancelAPI
 
 #include "UVSelectTool.generated.h"
 
@@ -39,7 +40,7 @@ public:
  * transform these elements.
  */
 UCLASS()
-class UVEDITORTOOLS_API UUVSelectTool : public UInteractiveTool, public IUVToolSupportsSelection
+class UVEDITORTOOLS_API UUVSelectTool : public UInteractiveTool, public IInteractiveToolNestedAcceptCancelAPI, public IUVToolSupportsSelection
 {
 	GENERATED_BODY()
 
@@ -53,6 +54,8 @@ public:
 		Targets = TargetsIn;
 	}
 
+	void SelectAll();
+
 	// Used by undo/redo.
 	FTransform GetGizmoTransform() const;
 	void SetGizmoTransform(const FTransform& NewTransform);
@@ -67,6 +70,11 @@ public:
 	virtual bool HasCancel() const override { return false; }
 	virtual bool HasAccept() const override { return false; }
 
+	// IInteractiveToolNestedAcceptCancelAPI
+	virtual bool SupportsNestedCancelCommand() override { return true; }
+	virtual bool CanCurrentlyNestedCancel() override;
+	virtual bool ExecuteNestedCancelCommand() override;
+
 protected:
 	virtual void OnSelectionChanged(bool bEmitChangeAllowed, uint32 SelectionChangeType);
 
