UVEditor: Enable support for snapping on gizmo initiated transforms. Adds a new toolbar in the 2D viewport that resembles the snapping controls found in the main level editor.

diff --git a/Engine/Plugins/Experimental/UVEditor/Source/UVEditor/Private/SUVEditor2DViewport.cpp b/Engine/Plugins/Experimental/UVEditor/Source/UVEditor/Private/SUVEditor2DViewport.cpp
--- a/Engine/Plugins/Experimental/UVEditor/Source/UVEditor/Private/SUVEditor2DViewport.cpp
+++ b/Engine/Plugins/Experimental/UVEditor/Source/UVEditor/Private/SUVEditor2DViewport.cpp
@@ -91,7 +91,8 @@ void SUVEditor2DViewport::RemoveOverlayWidget(TSharedRef<SWidget> OverlaidWidget
 TSharedPtr<SWidget> SUVEditor2DViewport::MakeViewportToolbar()
 {
 	return SNew(SUVEditor2DViewportToolBar)
-		.CommandList(CommandList);
+		.CommandList(CommandList)
+		.Viewport2DClient(StaticCastSharedPtr<FUVEditor2DViewportClient>(Client));
 }
 
 bool SUVEditor2DViewport::IsWidgetModeActive(UE::Widget::EWidgetMode Mode) const
diff --git a/Engine/Plugins/Experimental/UVEditor/Source/UVEditor/Private/SUVEditor2DViewportToolBar.cpp b/Engine/Plugins/Experimental/UVEditor/Source/UVEditor/Private/SUVEditor2DViewportToolBar.cpp
--- a/Engine/Plugins/Experimental/UVEditor/Source/UVEditor/Private/SUVEditor2DViewportToolBar.cpp
+++ b/Engine/Plugins/Experimental/UVEditor/Source/UVEditor/Private/SUVEditor2DViewportToolBar.cpp
@@ -7,11 +7,17 @@
 #include "UVEditorCommands.h"
 #include "Widgets/Layout/SBorder.h"
 #include "Widgets/SBoxPanel.h"
+#include "UVEditorUXSettings.h"
+#include "UVEditor2DViewportClient.h"
+#include "SViewportToolBarComboMenu.h"
+#include "Widgets/Layout/SUniformGridPanel.h"
 
+#define LOCTEXT_NAMESPACE "UVEditor2DViewportToolbar"
 
 void SUVEditor2DViewportToolBar::Construct(const FArguments& InArgs)
 {
 	CommandList = InArgs._CommandList;
+	Viewport2DClient = InArgs._Viewport2DClient;
 
 	const FMargin ToolbarSlotPadding(4.0f, 1.0f);
 
@@ -36,6 +42,13 @@ void SUVEditor2DViewportToolBar::Construct(const FArguments& InArgs)
 				MakeSelectionToolBar(InArgs._Extenders)
 			]
 
+	        + SHorizontalBox::Slot()
+			.Padding(ToolbarSlotPadding)
+			.HAlign(HAlign_Right)
+			[
+				MakeTransformToolBar(InArgs._Extenders)
+			]
+
 			+ SHorizontalBox::Slot()
 			.Padding(ToolbarSlotPadding)
 			.HAlign(HAlign_Right)
@@ -128,4 +141,331 @@ TSharedRef<SWidget> SUVEditor2DViewportToolBar::MakeGizmoToolBar(const TSharedPt
 	ToolbarBuilder.EndSection();
 
 	return ToolbarBuilder.MakeWidget();
-}
\ No newline at end of file
+}
+
+// This is mostly copied from the version found in STransformViewportToolbar.cpp, which provides the
+// transform/snapping controls for the level editor main window. We only want a subset of that 
+// functionality though and moreover we want to store the state in the UVEditor instead of the editor
+// global settings. So we make a few tweaks here from the original to control the exact snapping options
+// available and the activation methods.
+TSharedRef<SWidget> SUVEditor2DViewportToolBar::MakeTransformToolBar(const TSharedPtr< FExtender > InExtenders)
+{
+	FSlimHorizontalToolBarBuilder ToolbarBuilder(CommandList, FMultiBoxCustomization::None, InExtenders);
+
+	// Use a custom style
+	FName ToolBarStyle = "EditorViewportToolBar";
+	ToolbarBuilder.SetStyle(&FAppStyle::Get(), ToolBarStyle);
+	ToolbarBuilder.SetLabelVisibility(EVisibility::Collapsed);
+
+	// Transform controls cannot be focusable as it fights with the press space to change transform mode feature
+	ToolbarBuilder.SetIsFocusable(false);
+
+	ToolbarBuilder.BeginSection("LocationGridSnap");
+	{
+		// Grab the existing UICommand 
+		// TODO: Should we replace this with our own version at some point? 
+		// Same for the other commands below for rotation and scaling.
+		TSharedPtr<FUICommandInfo> Command = FEditorViewportCommands::Get().LocationGridSnap;
+
+		static FName PositionSnapName = FName(TEXT("PositionSnap"));
+
+		// Setup a GridSnapSetting with the UICommand
+		ToolbarBuilder.AddWidget(
+			SNew(SViewportToolBarComboMenu)
+			.Style(ToolBarStyle)
+			.IsChecked(this, &SUVEditor2DViewportToolBar::IsLocationGridSnapChecked)
+			.OnCheckStateChanged(this, &SUVEditor2DViewportToolBar::HandleToggleLocationGridSnap)
+			.Label(this, &SUVEditor2DViewportToolBar::GetLocationGridLabel)
+			.OnGetMenuContent(this, &SUVEditor2DViewportToolBar::FillLocationGridSnapMenu)
+			.ToggleButtonToolTip(Command->GetDescription())
+			.MenuButtonToolTip(LOCTEXT("LocationGridSnap_ToolTip", "Set the Translation Snap value"))
+			.Icon(Command->GetIcon())
+			.MinDesiredButtonWidth(24.0f)
+			.ParentToolBar(SharedThis(this)),
+			PositionSnapName,
+			false,
+			HAlign_Fill,
+
+			// explicitly specify what this widget should look like as a menu item
+			FNewMenuDelegate::CreateLambda([this, Command](FMenuBuilder& InMenuBuilder)
+				{
+					// TODO - debug why can't just use the Command / mapping isn't working 
+					InMenuBuilder.AddMenuEntry(Command);
+
+					InMenuBuilder.AddWrapperSubMenu(
+						LOCTEXT("GridSnapMenuSettings", "Translation Snap Settings"),
+						LOCTEXT("GridSnapMenuSettings_ToolTip", "Set the Translation Snap value"),
+						FOnGetContent::CreateSP(this, &SUVEditor2DViewportToolBar::FillLocationGridSnapMenu),
+						FSlateIcon(Command->GetIcon())
+					);
+				}
+		));
+	}
+	ToolbarBuilder.EndSection();
+
+	ToolbarBuilder.BeginSection("RotationGridSnap");
+	{
+		// Grab the existing UICommand 
+		TSharedPtr<FUICommandInfo> Command = FEditorViewportCommands::Get().RotationGridSnap;
+
+		static FName RotationSnapName = FName(TEXT("RotationSnap"));
+
+		// Setup a GridSnapSetting with the UICommand
+		ToolbarBuilder.AddWidget(
+			SNew(SViewportToolBarComboMenu)
+			.Style(ToolBarStyle)
+			.IsChecked(this, &SUVEditor2DViewportToolBar::IsRotationGridSnapChecked)
+			.OnCheckStateChanged(this, &SUVEditor2DViewportToolBar::HandleToggleRotationGridSnap)
+			.Label(this, &SUVEditor2DViewportToolBar::GetRotationGridLabel)
+			.OnGetMenuContent(this, &SUVEditor2DViewportToolBar::FillRotationGridSnapMenu)
+			.ToggleButtonToolTip(Command->GetDescription())
+			.MenuButtonToolTip(LOCTEXT("RotationGridSnap_ToolTip", "Set the Rotation Snap value"))
+			.Icon(Command->GetIcon())
+			.ParentToolBar(SharedThis(this)),
+			RotationSnapName,
+			false,
+			HAlign_Fill,
+
+			// explicitly specify what this widget should look like as a menu item
+			FNewMenuDelegate::CreateLambda([this, Command](FMenuBuilder& InMenuBuilder)
+				{
+					InMenuBuilder.AddMenuEntry(Command);
+
+					InMenuBuilder.AddWrapperSubMenu(
+						LOCTEXT("RotationGridSnapMenuSettings", "Rotation Snap Settings"),
+						LOCTEXT("RotationGridSnapMenuSettings_ToolTip", "Set the Rotation Snap value"),
+						FOnGetContent::CreateSP(this, &SUVEditor2DViewportToolBar::FillRotationGridSnapMenu),
+						FSlateIcon(Command->GetIcon())
+					);
+				}
+		));
+	}
+	ToolbarBuilder.EndSection();
+
+
+	ToolbarBuilder.BeginSection("ScaleGridSnap");
+	{
+		// Grab the existing UICommand 
+		TSharedPtr<FUICommandInfo> Command = FEditorViewportCommands::Get().ScaleGridSnap;
+
+		static FName ScaleSnapName = FName(TEXT("ScaleSnap"));
+
+		// Setup a GridSnapSetting with the UICommand
+		ToolbarBuilder.AddWidget(
+			SNew(SViewportToolBarComboMenu)
+			.Cursor(EMouseCursor::Default)
+			.Style(ToolBarStyle)
+			.IsChecked(this, &SUVEditor2DViewportToolBar::IsScaleGridSnapChecked)
+			.OnCheckStateChanged(this, &SUVEditor2DViewportToolBar::HandleToggleScaleGridSnap)
+			.Label(this, &SUVEditor2DViewportToolBar::GetScaleGridLabel)
+			.OnGetMenuContent(this, &SUVEditor2DViewportToolBar::FillScaleGridSnapMenu)
+			.ToggleButtonToolTip(Command->GetDescription())
+			.MenuButtonToolTip(LOCTEXT("ScaleGridSnap_ToolTip", "Set the Scaling Snap value"))
+			.Icon(Command->GetIcon())
+			.MinDesiredButtonWidth(24.0f)
+			.ParentToolBar(SharedThis(this)),
+			ScaleSnapName,
+			false,
+			HAlign_Fill,
+
+			// explicitly specify what this widget should look like as a menu item
+			FNewMenuDelegate::CreateLambda([this, Command](FMenuBuilder& InMenuBuilder)
+				{
+					InMenuBuilder.AddMenuEntry(Command);
+
+					InMenuBuilder.AddWrapperSubMenu(
+						LOCTEXT("ScaleGridSnapMenuSettings", "Scale Snap Settings"),
+						LOCTEXT("ScaleGridSnapMenuSettings_ToolTip", "Set the Scale Snap value"),
+						FOnGetContent::CreateSP(this, &SUVEditor2DViewportToolBar::FillScaleGridSnapMenu),
+						FSlateIcon(Command->GetIcon())
+					);
+				}
+		));
+	}
+	ToolbarBuilder.EndSection();
+
+	return ToolbarBuilder.MakeWidget();
+}
+
+// The following methods again mimic the patterns found in the STransformViewportToolbar.cpp,
+// to serve as a drop in replacement for the menu infrastructure above. These have been altered
+// to adjust the menu options and change how the snap settings are stored/forwarded to the UVEditor.
+
+FText SUVEditor2DViewportToolBar::GetLocationGridLabel() const
+{
+	return FText::AsNumber(Viewport2DClient->GetLocationGridSnapValue());
+}
+
+FText SUVEditor2DViewportToolBar::GetRotationGridLabel() const
+{
+	return FText::Format(LOCTEXT("GridRotation - Number - DegreeSymbol", "{0}\u00b0"), FText::AsNumber(Viewport2DClient->GetRotationGridSnapValue()));
+}
+
+FText SUVEditor2DViewportToolBar::GetScaleGridLabel() const
+{
+	FNumberFormattingOptions NumberFormattingOptions;
+	NumberFormattingOptions.MaximumFractionalDigits = 5;
+
+	const float CurGridAmount = Viewport2DClient->GetScaleGridSnapValue();
+	return FText::AsNumber(CurGridAmount, &NumberFormattingOptions);
+}
+
+TSharedRef<SWidget> SUVEditor2DViewportToolBar::FillLocationGridSnapMenu()
+{
+	TArray<float> GridSizes;
+	GridSizes.Reserve(FUVEditorUXSettings::MaxLocationSnapValue());
+	for (int32 Index = 0; Index < FUVEditorUXSettings::MaxLocationSnapValue(); ++Index)
+	{
+		GridSizes.Add(FUVEditorUXSettings::LocationSnapValue(Index));
+	}
+	return BuildLocationGridCheckBoxList("Snap", LOCTEXT("LocationSnapText", "Snap Delta Distances"), GridSizes);
+}
+
+TSharedRef<SWidget> SUVEditor2DViewportToolBar::BuildLocationGridCheckBoxList(FName InExtentionHook, const FText& InHeading, const TArray<float>& InGridSizes) const
+{
+	const ULevelEditorViewportSettings* ViewportSettings = GetDefault<ULevelEditorViewportSettings>();
+
+	const bool bShouldCloseWindowAfterMenuSelection = true;
+	FMenuBuilder LocationGridMenuBuilder(bShouldCloseWindowAfterMenuSelection, CommandList);
+
+	LocationGridMenuBuilder.BeginSection(InExtentionHook, InHeading);
+	for (int32 CurGridSizeIndex = 0; CurGridSizeIndex < InGridSizes.Num(); ++CurGridSizeIndex)
+	{
+		const float CurGridSize = InGridSizes[CurGridSizeIndex];
+
+		LocationGridMenuBuilder.AddMenuEntry(
+			FText::AsNumber(CurGridSize),
+			FText::Format(LOCTEXT("LocationGridSize_ToolTip", "Sets snap delta to {0}"), FText::AsNumber(CurGridSize)),
+			FSlateIcon(),
+			FUIAction(FExecuteAction::CreateLambda([this, CurGridSize]() {Viewport2DClient->SetLocationGridSnapValue(CurGridSize); }),
+				FCanExecuteAction(),
+				FIsActionChecked::CreateLambda([this, CurGridSize]() { return FMath::IsNearlyEqual(Viewport2DClient->GetLocationGridSnapValue(), CurGridSize); })),
+			NAME_None,
+			EUserInterfaceActionType::RadioButton);
+	}
+	LocationGridMenuBuilder.EndSection();
+
+	return LocationGridMenuBuilder.MakeWidget();
+}
+
+TSharedRef<SWidget> SUVEditor2DViewportToolBar::FillRotationGridSnapMenu()
+{
+	const ULevelEditorViewportSettings* ViewportSettings = GetDefault<ULevelEditorViewportSettings>();
+
+	return SNew(SUniformGridPanel)
+
+		+ SUniformGridPanel::Slot(0, 0)
+		[
+			BuildRotationGridCheckBoxList("Common", LOCTEXT("RotationCommonText", "Common"), ViewportSettings->CommonRotGridSizes)
+		]
+
+	+ SUniformGridPanel::Slot(1, 0)
+		[
+			BuildRotationGridCheckBoxList("Div360", LOCTEXT("RotationDivisions360DegreesText", "Divisions of 360\u00b0"), ViewportSettings->DivisionsOf360RotGridSizes)
+		];
+}
+
+TSharedRef<SWidget> SUVEditor2DViewportToolBar::BuildRotationGridCheckBoxList(FName InExtentionHook, const FText& InHeading, const TArray<float>& InGridSizes) const
+{
+	const bool bShouldCloseWindowAfterMenuSelection = true;
+	FMenuBuilder RotationGridMenuBuilder(bShouldCloseWindowAfterMenuSelection, CommandList);
+
+	RotationGridMenuBuilder.BeginSection(InExtentionHook, InHeading);
+	for (int32 CurGridAngleIndex = 0; CurGridAngleIndex < InGridSizes.Num(); ++CurGridAngleIndex)
+	{
+		const float CurGridAngle = InGridSizes[CurGridAngleIndex];
+
+		FText MenuName = FText::Format(LOCTEXT("RotationGridAngle", "{0}\u00b0"), FText::AsNumber(CurGridAngle)); /*degree symbol*/
+		FText ToolTipText = FText::Format(LOCTEXT("RotationGridAngle_ToolTip", "Sets rotation snap angle to {0}"), MenuName); /*degree symbol*/
+
+		RotationGridMenuBuilder.AddMenuEntry(
+			MenuName,
+			ToolTipText,
+			FSlateIcon(),
+			FUIAction(FExecuteAction::CreateLambda([this, CurGridAngle]() {Viewport2DClient->SetRotationGridSnapValue(CurGridAngle); }),
+				FCanExecuteAction(),
+				FIsActionChecked::CreateLambda([this, CurGridAngle]() {return FMath::IsNearlyEqual(Viewport2DClient->GetRotationGridSnapValue(), CurGridAngle); })),
+			NAME_None,
+			EUserInterfaceActionType::RadioButton);
+	}
+	RotationGridMenuBuilder.EndSection();
+
+	return RotationGridMenuBuilder.MakeWidget();
+}
+
+TSharedRef<SWidget> SUVEditor2DViewportToolBar::FillScaleGridSnapMenu()
+{
+	const ULevelEditorViewportSettings* ViewportSettings = GetDefault<ULevelEditorViewportSettings>();
+	const bool bShouldCloseWindowAfterMenuSelection = true;
+
+	FNumberFormattingOptions NumberFormattingOptions;
+	NumberFormattingOptions.MaximumFractionalDigits = 5;
+
+	FMenuBuilder ScaleGridMenuBuilder(bShouldCloseWindowAfterMenuSelection, CommandList);
+
+	ScaleGridMenuBuilder.BeginSection("ScaleSnapOptions", LOCTEXT("ScaleSnapOptions", "Scale Snap"));
+
+	for (int32 CurGridAmountIndex = 0; CurGridAmountIndex < ViewportSettings->ScalingGridSizes.Num(); ++CurGridAmountIndex)
+	{
+		const float CurGridAmount = ViewportSettings->ScalingGridSizes[CurGridAmountIndex];
+
+		FText MenuText;
+		FText ToolTipText;
+
+		if (GEditor->UsePercentageBasedScaling())
+		{
+			MenuText = FText::AsPercent(CurGridAmount / 100.0f, &NumberFormattingOptions);
+			ToolTipText = FText::Format(LOCTEXT("ScaleGridAmountOld_ToolTip", "Snaps scale values to {0}"), MenuText);
+		}
+		else
+		{
+			MenuText = FText::AsNumber(CurGridAmount, &NumberFormattingOptions);
+			ToolTipText = FText::Format(LOCTEXT("ScaleGridAmount_ToolTip", "Snaps scale values to increments of {0}"), MenuText);
+		}
+
+		ScaleGridMenuBuilder.AddMenuEntry(
+			MenuText,
+			ToolTipText,
+			FSlateIcon(),
+			FUIAction(FExecuteAction::CreateLambda([this, CurGridAmount]() {Viewport2DClient->SetScaleGridSnapValue(CurGridAmount); }),
+				FCanExecuteAction(),
+				FIsActionChecked::CreateLambda([this, CurGridAmount]() {return FMath::IsNearlyEqual(Viewport2DClient->GetScaleGridSnapValue(), CurGridAmount); })),
+			NAME_None,
+			EUserInterfaceActionType::RadioButton);
+	}
+	ScaleGridMenuBuilder.EndSection();
+
+	return ScaleGridMenuBuilder.MakeWidget();
+}
+
+ECheckBoxState SUVEditor2DViewportToolBar::IsLocationGridSnapChecked() const
+{
+	return Viewport2DClient->GetLocationGridSnapEnabled() ? ECheckBoxState::Checked : ECheckBoxState::Unchecked;
+}
+
+ECheckBoxState SUVEditor2DViewportToolBar::IsRotationGridSnapChecked() const
+{
+	return Viewport2DClient->GetRotationGridSnapEnabled() ? ECheckBoxState::Checked : ECheckBoxState::Unchecked;
+}
+
+ECheckBoxState SUVEditor2DViewportToolBar::IsScaleGridSnapChecked() const
+{
+	return Viewport2DClient->GetScaleGridSnapEnabled() ? ECheckBoxState::Checked : ECheckBoxState::Unchecked;
+}
+
+void SUVEditor2DViewportToolBar::HandleToggleLocationGridSnap(ECheckBoxState InState)
+{
+	Viewport2DClient->SetLocationGridSnapEnabled(InState == ECheckBoxState::Checked);
+}
+
+void SUVEditor2DViewportToolBar::HandleToggleRotationGridSnap(ECheckBoxState InState)
+{
+	Viewport2DClient->SetRotationGridSnapEnabled(InState == ECheckBoxState::Checked);
+}
+
+void SUVEditor2DViewportToolBar::HandleToggleScaleGridSnap(ECheckBoxState InState)
+{
+	Viewport2DClient->SetScaleGridSnapEnabled(InState == ECheckBoxState::Checked);
+}
+
+#undef LOCTEXT_NAMESPACE
\ No newline at end of file
diff --git a/Engine/Plugins/Experimental/UVEditor/Source/UVEditor/Private/UVEditor2DViewportClient.cpp b/Engine/Plugins/Experimental/UVEditor/Source/UVEditor/Private/UVEditor2DViewportClient.cpp
--- a/Engine/Plugins/Experimental/UVEditor/Source/UVEditor/Private/UVEditor2DViewportClient.cpp
+++ b/Engine/Plugins/Experimental/UVEditor/Source/UVEditor/Private/UVEditor2DViewportClient.cpp
@@ -314,6 +314,66 @@ bool FUVEditor2DViewportClient::AreWidgetButtonsEnabled() const
 	return ViewportButtonsAPI && ViewportButtonsAPI->AreGizmoButtonsEnabled();
 }
 
+void FUVEditor2DViewportClient::SetLocationGridSnapEnabled(bool bEnabled)
+{
+	ViewportButtonsAPI->ToggleSnapEnabled(UUVToolViewportButtonsAPI::ESnapTypeFlag::Location);
+}
+
+bool FUVEditor2DViewportClient::GetLocationGridSnapEnabled()
+{
+	return ViewportButtonsAPI->GetSnapEnabled(UUVToolViewportButtonsAPI::ESnapTypeFlag::Location);
+}
+
+void FUVEditor2DViewportClient::SetLocationGridSnapValue(float SnapValue)
+{
+	ViewportButtonsAPI->SetSnapValue(UUVToolViewportButtonsAPI::ESnapTypeFlag::Location, SnapValue);
+}
+
+float FUVEditor2DViewportClient::GetLocationGridSnapValue()
+{
+	return 	ViewportButtonsAPI->GetSnapValue(UUVToolViewportButtonsAPI::ESnapTypeFlag::Location);
+}
+
+void FUVEditor2DViewportClient::SetRotationGridSnapEnabled(bool bEnabled)
+{
+	ViewportButtonsAPI->ToggleSnapEnabled(UUVToolViewportButtonsAPI::ESnapTypeFlag::Rotation);
+}
+
+bool FUVEditor2DViewportClient::GetRotationGridSnapEnabled()
+{
+	return ViewportButtonsAPI->GetSnapEnabled(UUVToolViewportButtonsAPI::ESnapTypeFlag::Rotation);
+}
+
+void FUVEditor2DViewportClient::SetRotationGridSnapValue(float SnapValue)
+{
+	ViewportButtonsAPI->SetSnapValue(UUVToolViewportButtonsAPI::ESnapTypeFlag::Rotation, SnapValue);
+}
+
+float FUVEditor2DViewportClient::GetRotationGridSnapValue()
+{
+	return ViewportButtonsAPI->GetSnapValue(UUVToolViewportButtonsAPI::ESnapTypeFlag::Rotation);
+}
+
+void FUVEditor2DViewportClient::SetScaleGridSnapEnabled(bool bEnabled)
+{
+	ViewportButtonsAPI->ToggleSnapEnabled(UUVToolViewportButtonsAPI::ESnapTypeFlag::Scale);
+}
+
+bool FUVEditor2DViewportClient::GetScaleGridSnapEnabled()
+{
+	return ViewportButtonsAPI->GetSnapEnabled(UUVToolViewportButtonsAPI::ESnapTypeFlag::Scale);
+}
+
+void FUVEditor2DViewportClient::SetScaleGridSnapValue(float SnapValue)
+{
+	ViewportButtonsAPI->SetSnapValue(UUVToolViewportButtonsAPI::ESnapTypeFlag::Scale, SnapValue);
+}
+
+float FUVEditor2DViewportClient::GetScaleGridSnapValue()
+{
+	return ViewportButtonsAPI->GetSnapValue(UUVToolViewportButtonsAPI::ESnapTypeFlag::Scale);
+}
+
 bool FUVEditor2DViewportClient::CanSetWidgetMode(UE::Widget::EWidgetMode NewMode) const
 {
 	if (!AreWidgetButtonsEnabled())
diff --git a/Engine/Plugins/Experimental/UVEditor/Source/UVEditor/Public/SUVEditor2DViewportToolBar.h b/Engine/Plugins/Experimental/UVEditor/Source/UVEditor/Public/SUVEditor2DViewportToolBar.h
--- a/Engine/Plugins/Experimental/UVEditor/Source/UVEditor/Public/SUVEditor2DViewportToolBar.h
+++ b/Engine/Plugins/Experimental/UVEditor/Source/UVEditor/Public/SUVEditor2DViewportToolBar.h
@@ -5,10 +5,12 @@
 #include "CoreMinimal.h"
 
 #include "SViewportToolBar.h"
+#include "Styling/SlateTypes.h"
 
 class FExtender;
 class FUICommandList;
 class SUVEditor2DViewport;
+class FUVEditor2DViewportClient;
 
 /**
  * Toolbar that shows up at the top of the 2d viewport (has gizmo controls)
@@ -19,6 +21,7 @@ public:
 	SLATE_BEGIN_ARGS(SUVEditor2DViewportToolBar) {}
 		SLATE_ARGUMENT(TSharedPtr<FUICommandList>, CommandList)
 		SLATE_ARGUMENT(TSharedPtr<FExtender>, Extenders)
+		SLATE_ARGUMENT(TSharedPtr<FUVEditor2DViewportClient>, Viewport2DClient)
 	SLATE_END_ARGS()
 
 	void Construct(const FArguments& InArgs);
@@ -26,7 +29,40 @@ public:
 private:
 	TSharedRef<SWidget> MakeSelectionToolBar(const TSharedPtr<FExtender> InExtenders);
 	TSharedRef<SWidget> MakeGizmoToolBar(const TSharedPtr<FExtender> InExtenders);
+	TSharedRef<SWidget> MakeTransformToolBar(const TSharedPtr< FExtender > InExtenders);
+
+	TSharedRef<SWidget> BuildLocationGridCheckBoxList(FName InExtentionHook, const FText& InHeading, const TArray<float>& InGridSizes) const;
+	TSharedRef<SWidget> BuildRotationGridCheckBoxList(FName InExtentionHook, const FText& InHeading, const TArray<float>& InGridSizes) const;
+
+
+	/** Grid snap label callbacks */
+	FText GetLocationGridLabel() const;
+	FText GetRotationGridLabel() const;	
+	FText GetScaleGridLabel() const;
+
+	/** GridSnap menu construction callbacks */
+	TSharedRef<SWidget> FillLocationGridSnapMenu();
+	TSharedRef<SWidget> FillRotationGridSnapMenu();
+	TSharedRef<SWidget> FillScaleGridSnapMenu();
+
+	/** Grid Snap checked state callbacks */
+	ECheckBoxState IsLocationGridSnapChecked() const;
+	ECheckBoxState IsRotationGridSnapChecked() const;
+	ECheckBoxState IsScaleGridSnapChecked() const;
+
+	/** Grid snap toggle handlers */
+	void HandleToggleLocationGridSnap(ECheckBoxState InState);
+	void HandleToggleRotationGridSnap(ECheckBoxState InState);	
+	void HandleToggleScaleGridSnap(ECheckBoxState InState);
 
 	/** Command list */
 	TSharedPtr<FUICommandList> CommandList;
+
+	/** Client - For connecting snap controls */
+	TSharedPtr<FUVEditor2DViewportClient> Viewport2DClient;
+
+	/** Current Snap Indices */
+	int32 CurLocationSnapIndex;
+	int32 CurRotationSnapIndex;
+	int32 CurScaleSnapIndex;
 };
diff --git a/Engine/Plugins/Experimental/UVEditor/Source/UVEditor/Public/UVEditor2DViewportClient.h b/Engine/Plugins/Experimental/UVEditor/Source/UVEditor/Public/UVEditor2DViewportClient.h
--- a/Engine/Plugins/Experimental/UVEditor/Source/UVEditor/Public/UVEditor2DViewportClient.h
+++ b/Engine/Plugins/Experimental/UVEditor/Source/UVEditor/Public/UVEditor2DViewportClient.h
@@ -29,6 +29,19 @@ public:
 	UUVToolViewportButtonsAPI::ESelectionMode GetSelectionMode() const;
 	bool AreWidgetButtonsEnabled() const;
 
+	void SetLocationGridSnapEnabled(bool bEnabled);
+	bool GetLocationGridSnapEnabled();
+	void SetLocationGridSnapValue(float SnapValue);
+	float GetLocationGridSnapValue();
+	void SetRotationGridSnapEnabled(bool bEnabled);
+	bool GetRotationGridSnapEnabled();
+	void SetRotationGridSnapValue(float SnapValue);
+	float GetRotationGridSnapValue();
+	void SetScaleGridSnapEnabled(bool bEnabled);
+	bool GetScaleGridSnapEnabled();
+	void SetScaleGridSnapValue(float SnapValue);
+	float GetScaleGridSnapValue();
+
 	// FEditorViewportClient
 	virtual bool InputKey(const FInputKeyEventArgs& EventArgs) override;
 
diff --git a/Engine/Plugins/Experimental/UVEditor/Source/UVEditorTools/Private/UVEditorUXSettings.cpp b/Engine/Plugins/Experimental/UVEditor/Source/UVEditorTools/Private/UVEditorUXSettings.cpp
--- a/Engine/Plugins/Experimental/UVEditor/Source/UVEditorTools/Private/UVEditorUXSettings.cpp
+++ b/Engine/Plugins/Experimental/UVEditor/Source/UVEditorTools/Private/UVEditorUXSettings.cpp
@@ -3,6 +3,8 @@
 #include "UVEditorUXSettings.h"
 #include "Math/Color.h"
 
+#define LOCTEXT_NAMESPACE "UVEditorUXSettings"
+
 const float FUVEditorUXSettings::UVMeshScalingFactor(1000.0);
 const float FUVEditorUXSettings::CameraFarPlaneWorldZ(-10.0);
 const float FUVEditorUXSettings::CameraNearPlaneProportionZ(0.8); // Top layer, equivalent to depth bias 80
@@ -135,4 +137,31 @@ FVector3d FUVEditorUXSettings::UVToVertPosition(const FVector2f& UV)
 FVector2f FUVEditorUXSettings::VertPositionToUV(const FVector3d& VertPosition)
 {
 	return FVector2f(VertPosition.Y / UVMeshScalingFactor, 1 - (VertPosition.X / UVMeshScalingFactor));
-};
\ No newline at end of file
+};
+
+float FUVEditorUXSettings::LocationSnapValue(int32 LocationSnapMenuIndex)
+{
+	switch (LocationSnapMenuIndex)
+	{
+	case 0:
+		return 1.0;
+	case 1:
+		return 0.5;
+	case 2:
+		return 0.25;
+	case 3:
+		return 0.125;
+	case 4:
+		return 0.0625;
+	default:
+		ensure(false);
+		return 0;
+	}
+}
+
+int32 FUVEditorUXSettings::MaxLocationSnapValue()
+{
+	return 5;
+}
+
+#undef LOCTEXT_NAMESPACE 
diff --git a/Engine/Plugins/Experimental/UVEditor/Source/UVEditorTools/Private/UVSelectTool.cpp b/Engine/Plugins/Experimental/UVEditor/Source/UVEditorTools/Private/UVSelectTool.cpp
--- a/Engine/Plugins/Experimental/UVEditor/Source/UVEditorTools/Private/UVSelectTool.cpp
+++ b/Engine/Plugins/Experimental/UVEditor/Source/UVEditorTools/Private/UVSelectTool.cpp
@@ -486,8 +486,65 @@ void UUVSelectTool::GizmoTransformChanged(UTransformProxy* Proxy, FTransform Tra
 
 	if (!DeltaTransform.GetTranslation().IsNearlyZero() || !DeltaTransform.GetRotation().IsIdentity() || Transform.GetScale3D() != FVector::One())
 	{
-		UnappliedGizmoTransform = Transform;
-		bGizmoTransformNeedsApplication = true;
+	
+		FTransform SnappedDeltaTransform = DeltaTransform;		
+		bool bSnapExecuted = false;
+
+		if (!DeltaTransform.GetTranslation().IsNearlyZero() && 
+			ViewportButtonsAPI->GetSnapEnabled(UUVToolViewportButtonsAPI::ESnapTypeFlag::Location))
+		{
+			float SnapValue = ViewportButtonsAPI->GetSnapValue(UUVToolViewportButtonsAPI::ESnapTypeFlag::Location);
+			FVector Translation = DeltaTransform.GetTranslation();
+			Translation = Translation.GridSnap(SnapValue * FUVEditorUXSettings::UVMeshScalingFactor);
+			Translation = Transform.GetRotation() * Translation; // Apply the current rotated frame, so we translate in local coordinates.
+			SnappedDeltaTransform.SetTranslation(Translation);
+			bSnapExecuted = true;
+		}
+
+		if (!DeltaTransform.GetRotation().IsIdentity() &&
+			ViewportButtonsAPI->GetSnapEnabled(UUVToolViewportButtonsAPI::ESnapTypeFlag::Rotation))
+		{
+			float SnapValue = ViewportButtonsAPI->GetSnapValue(UUVToolViewportButtonsAPI::ESnapTypeFlag::Rotation);
+			FQuat DeltaRotation = DeltaTransform.GetRotation();
+			FVector DeltaPosition = InitialGizmoFrame.Origin;
+
+			FRotator Rotator(DeltaTransform.GetRotation());
+			FRotator RotGrid(SnapValue, SnapValue, SnapValue);
+			Rotator = Rotator.GridSnap(RotGrid);
+			DeltaPosition = ( Rotator.Quaternion() * -DeltaPosition ) + DeltaPosition;
+
+			SnappedDeltaTransform.SetRotation(Rotator.Quaternion());
+			SnappedDeltaTransform.SetLocation(DeltaPosition);
+			bSnapExecuted = true;
+		}
+
+		if (Transform.GetScale3D() != FVector::One() &&
+			ViewportButtonsAPI->GetSnapEnabled(UUVToolViewportButtonsAPI::ESnapTypeFlag::Scale))
+		{
+			float SnapValue = ViewportButtonsAPI->GetSnapValue(UUVToolViewportButtonsAPI::ESnapTypeFlag::Scale);
+			FVector DeltaScale = DeltaTransform.GetScale3D();
+			FVector DeltaPosition = InitialGizmoFrame.Origin;			
+			DeltaScale = DeltaScale.GridSnap(SnapValue);
+			DeltaPosition = (DeltaScale * -DeltaPosition) + DeltaPosition;
+
+			SnappedDeltaTransform.SetScale3D(DeltaScale);
+			SnappedDeltaTransform.SetLocation(DeltaPosition);
+			bSnapExecuted = true;
+		}
+
+
+		FTransform SnappedTransform = InitialGizmoFrame.ToFTransform() * SnappedDeltaTransform;
+		
+		if (bSnapExecuted)
+		{
+			UnappliedGizmoTransform = SnappedTransform;
+			TransformGizmo->ReinitializeGizmoTransform(SnappedTransform);
+		}
+		else
+		{
+			UnappliedGizmoTransform = Transform;
+		}
+		bGizmoTransformNeedsApplication = true;		
 	}	
 }
 
@@ -553,7 +610,12 @@ void UUVSelectTool::ApplyGizmoTransform()
 	// easier to zoom in and out without running into issues, the measure of the distance across which we typically
 	// drag the handles is too high to be convenient. Until we make the scaling invariant to units/distance from
 	// target, we use this hack.
-	TransformToApply.SetScale(FVector::One() + (UnappliedGizmoTransform.GetScale3D() - FVector::One()) / 10);
+	// Note: If we're using the snapping feature we don't want to mess with the values on scale here, since the
+	// snapping has already been performed. If we use the hack, we won't have the exact snapped scale the user expects.
+	if (!ViewportButtonsAPI->GetSnapEnabled(UUVToolViewportButtonsAPI::ESnapTypeFlag::Scale))
+	{
+		TransformToApply.SetScale(FVector::One() + (UnappliedGizmoTransform.GetScale3D() - FVector::One()) / 10);
+	}
 
 	for (int32 SelectionIndex = 0; SelectionIndex < CurrentSelections.Num(); ++SelectionIndex)
 	{
diff --git a/Engine/Plugins/Experimental/UVEditor/Source/UVEditorTools/Public/ContextObjects/UVToolViewportButtonsAPI.h b/Engine/Plugins/Experimental/UVEditor/Source/UVEditorTools/Public/ContextObjects/UVToolViewportButtonsAPI.h
--- a/Engine/Plugins/Experimental/UVEditor/Source/UVEditorTools/Public/ContextObjects/UVToolViewportButtonsAPI.h
+++ b/Engine/Plugins/Experimental/UVEditor/Source/UVEditorTools/Public/ContextObjects/UVToolViewportButtonsAPI.h
@@ -23,6 +23,14 @@ public:
 		Transform
 	};
 
+	enum class ESnapTypeFlag : uint8
+	{
+		NoSnap = 0,
+		Location = 1 << 0,
+		Rotation = 1 << 1,
+		Scale = 1 << 2
+	};
+
 	using ESelectionMode = UUVToolSelectionAPI::EUVEditorSelectionMode;
 
 	void SetGizmoButtonsEnabled(bool bOn)
@@ -93,9 +101,59 @@ public:
 		OnGizmoModeChange.RemoveAll(DeadTool);
 		OnSelectionModeChange.RemoveAll(DeadTool);
 	}
+	
+	void ToggleSnapEnabled(ESnapTypeFlag SnapMode)
+	{
+		SnapEnabled ^= (uint8)SnapMode;
+	}
+
+	bool GetSnapEnabled(ESnapTypeFlag SnapMode) const
+	{
+		return SnapEnabled & (uint8)SnapMode;
+	}
+
+	void SetSnapValue(ESnapTypeFlag SnapMode, float SnapValue)
+	{
+		switch (SnapMode)
+		{
+		case ESnapTypeFlag::Location:
+			LocationSnap = SnapValue;
+			break;
+		case ESnapTypeFlag::Rotation:
+			RotationSnap = SnapValue;
+			break;
+		case ESnapTypeFlag::Scale:
+			ScaleSnap = SnapValue;
+			break;
+		default:
+			ensure(false);
+		}
+	}
+
+	float GetSnapValue(ESnapTypeFlag SnapMode) const 
+	{
+		switch (SnapMode)
+		{
+		case ESnapTypeFlag::Location:
+			return LocationSnap;			
+		case ESnapTypeFlag::Rotation:
+			return RotationSnap;
+		case ESnapTypeFlag::Scale:
+			return ScaleSnap;
+		default:
+			ensure(false);
+			return 0.0;
+		}
+	}
+
 
 protected:
 
+	int32 SnapEnabled = (uint8)ESnapTypeFlag::NoSnap;
+	float LocationSnap = 1.0;
+	float RotationSnap = 5.0;
+	float ScaleSnap = 1.0;
+
 	bool bGizmoButtonsEnabled = false;
 	EGizmoMode GizmoMode = EGizmoMode::Select;
 	bool bSelectionButtonsEnabled = false;
diff --git a/Engine/Plugins/Experimental/UVEditor/Source/UVEditorTools/Public/UVEditorUXSettings.h b/Engine/Plugins/Experimental/UVEditor/Source/UVEditorTools/Public/UVEditorUXSettings.h
--- a/Engine/Plugins/Experimental/UVEditor/Source/UVEditorTools/Public/UVEditorUXSettings.h
+++ b/Engine/Plugins/Experimental/UVEditor/Source/UVEditorTools/Public/UVEditorUXSettings.h
@@ -131,6 +131,12 @@ public:
 
 	static TAutoConsoleVariable<int32> CVarEnablePrototypeUDIMSupport;
 
+	//--------------------------------
+	// Values for Snapping
+	//--------------------------------
+	static float LocationSnapValue(int32 LocationSnapMenuIndex);
+	static int32 MaxLocationSnapValue();
+
 private:
 
 	FUVEditorUXSettings();
