UVEditor: Enable support for SpectralConformal unwrapping mode in the UVEditor Unwrap tool.

diff --git a/Engine/Plugins/Experimental/UVEditor/Source/UVEditorTools/Private/Operators/UVEditorRecomputeUVsOp.cpp b/Engine/Plugins/Experimental/UVEditor/Source/UVEditorTools/Private/Operators/UVEditorRecomputeUVsOp.cpp
--- a/Engine/Plugins/Experimental/UVEditor/Source/UVEditorTools/Private/Operators/UVEditorRecomputeUVsOp.cpp
+++ b/Engine/Plugins/Experimental/UVEditor/Source/UVEditorTools/Private/Operators/UVEditorRecomputeUVsOp.cpp
@@ -286,6 +286,15 @@ bool FUVEditorRecomputeUVsOp::CalculateResult_Basic(FProgressCancel* Progress)
 					UVEditor.ScaleUVAreaTo3DArea(ComponentTris, true);
 				}
 				break;
+			case EUVEditorRecomputeUVsUnwrapType::SpectralConformal:
+			{
+				bComponentSolved[k] = UVEditor.SetTriangleUVsFromFreeBoundarySpectralConformal(ComponentTris, bUseExistingUVTopology, bPreserveIrregularity);
+				if (bComponentSolved[k])
+				{
+					UVEditor.ScaleUVAreaTo3DArea(ComponentTris, true);
+				}
+				break;
+			}
 			}
 
 			if (bComponentSolved[k])
@@ -534,6 +543,11 @@ TUniquePtr<FDynamicMeshOperator> UUVEditorRecomputeUVsOpFactory::MakeNewOperator
 	case EUVEditorRecomputeUVsPropertiesUnwrapType::IslandMerging:
 		RecomputeUVsOp->UnwrapType = EUVEditorRecomputeUVsUnwrapType::ExpMap;
 		RecomputeUVsOp->bMergingOptimization = true;
+		break;
+	case EUVEditorRecomputeUVsPropertiesUnwrapType::SpectralConformal:
+		RecomputeUVsOp->UnwrapType = EUVEditorRecomputeUVsUnwrapType::SpectralConformal;
+		RecomputeUVsOp->bMergingOptimization = false;
+		break;
 	}
 
 	RecomputeUVsOp->bAutoRotate = (Settings->AutoRotation == EUVEditorRecomputeUVsToolOrientationMode::MinBounds);
@@ -541,6 +555,8 @@ TUniquePtr<FDynamicMeshOperator> UUVEditorRecomputeUVsOpFactory::MakeNewOperator
 	RecomputeUVsOp->NormalSmoothingRounds = Settings->SmoothingSteps;
 	RecomputeUVsOp->NormalSmoothingAlpha = Settings->SmoothingAlpha;
 
+	RecomputeUVsOp->bPreserveIrregularity = Settings->bPreserveIrregularity;
+
 	RecomputeUVsOp->MergingThreshold = Settings->MergingDistortionThreshold;
 	RecomputeUVsOp->MaxNormalDeviationDeg = Settings->MergingAngleThreshold;
 
diff --git a/Engine/Plugins/Experimental/UVEditor/Source/UVEditorTools/Public/Operators/UVEditorRecomputeUVsOp.h b/Engine/Plugins/Experimental/UVEditor/Source/UVEditorTools/Public/Operators/UVEditorRecomputeUVsOp.h
--- a/Engine/Plugins/Experimental/UVEditor/Source/UVEditorTools/Public/Operators/UVEditorRecomputeUVsOp.h
+++ b/Engine/Plugins/Experimental/UVEditor/Source/UVEditorTools/Public/Operators/UVEditorRecomputeUVsOp.h
@@ -29,8 +29,10 @@ enum class EUVEditorRecomputeUVsPropertiesUnwrapType
 	ExpMap = 0 UMETA(DisplayName = "ExpMap"),
 	/** Conformal UV flattening is increasingly expensive on large islands but reduces distortion */
 	Conformal = 1,
+	/** Compared to the default Conformal method does not pin two vertices along the boundary which reduces the distortion but is more expensive to compute. */
+	SpectralConformal = 2,
 	/** UV islands will be merged into larger islands if it does not increase stretching and distortion beyond defined limits */
-	IslandMerging = 2
+	IslandMerging = 3
 };
 
 
@@ -93,6 +95,10 @@ public:
 	UPROPERTY(EditAnywhere, Category = "UV Unwrap")
 	EUVEditorRecomputeUVsToolOrientationMode AutoRotation = EUVEditorRecomputeUVsToolOrientationMode::MinBounds;
 
+	/**  If enabled, reduces distortion for meshes with triangles of vastly different sizes, This is only enabled if the Unwrap Type is set to Spectral Conformal. */
+	UPROPERTY(EditAnywhere, Category = "UV Unwrap", meta = (EditCondition = "UnwrapType == EUVEditorRecomputeUVsPropertiesUnwrapType::SpectralConformal"))
+	bool bPreserveIrregularity = true;
+
 	/** Number of smoothing steps to apply; this slightly increases distortion but produces more stable results. This is only enabled if the Unwrap Type is set to ExpMap or Island Merging. */
 	UPROPERTY(EditAnywhere, Category = "UV Unwrap", meta = (UIMin = "0", UIMax = "25", ClampMin = "0", ClampMax = "1000",
 		EditCondition = "UnwrapType == EUVEditorRecomputeUVsPropertiesUnwrapType::ExpMap || UnwrapType == EUVEditorRecomputeUVsPropertiesUnwrapType::IslandMerging"))
@@ -146,7 +152,8 @@ namespace Geometry
 enum class EUVEditorRecomputeUVsUnwrapType
 {
 	ExpMap = 0,
-	ConformalFreeBoundary = 1
+	ConformalFreeBoundary = 1,
+	SpectralConformal = 2
 };
 
 
@@ -194,6 +201,11 @@ public:
 	EUVEditorRecomputeUVsUnwrapType UnwrapType = EUVEditorRecomputeUVsUnwrapType::ExpMap;
 	EUVEditorRecomputeUVsIslandMode IslandMode = EUVEditorRecomputeUVsIslandMode::PolyGroups;
 
+	//
+	// Spectral Conformal Map Options 
+	//
+	bool bPreserveIrregularity = false;
+
 	// 
 	// ExpMap Options
 	//
