UVEditor: Optimized UUVSelectTool::OnSelectionChanged()

diff --git a/Engine/Plugins/Experimental/UVEditor/Source/UVEditorTools/Private/UVSelectTool.cpp b/Engine/Plugins/Experimental/UVEditor/Source/UVEditorTools/Private/UVSelectTool.cpp
--- a/Engine/Plugins/Experimental/UVEditor/Source/UVEditorTools/Private/UVSelectTool.cpp
+++ b/Engine/Plugins/Experimental/UVEditor/Source/UVEditorTools/Private/UVSelectTool.cpp
@@ -17,6 +17,7 @@
 #include "UVEditorUXSettings.h"
 
 #include "ToolTargetManager.h"
+#include "Algo/Unique.h"
 
 #define LOCTEXT_NAMESPACE "UUVSelectTool"
 
@@ -383,55 +384,67 @@ void UUVSelectTool::OnSelectionChanged(bool, uint32 SelectionChangeType)
 
 	for (int32 SelectionIndex = 0; SelectionIndex < CurrentSelections.Num(); ++SelectionIndex)
 	{
-		FUVToolSelection& Selection = CurrentSelections[SelectionIndex];
-		UUVEditorToolMeshInput* Target = Selection.Target.Get();
-		FDynamicMesh3* UnwrapMesh = Target->UnwrapCanonical.Get();
+		const FUVToolSelection& Selection = CurrentSelections[SelectionIndex];
+		const UUVEditorToolMeshInput* Target = Selection.Target.Get();
+		const FDynamicMesh3* UnwrapMesh = Target->UnwrapCanonical.Get();
 
-		// Note the selected vids
-		TSet<int32> VidSet;
-		TSet<int32> TidSet;
+		// Note the selected vertex IDs and triangle IDs
+		TArray<int32> VidSet;
+		TArray<int32> TidSet;
 		if (Selection.Type == FUVToolSelection::EType::Triangle)
 		{
 			TRACE_CPUPROFILER_EVENT_SCOPE(Triangle);
 
-			for (int32 Tid : Selection.SelectedIDs)
+			VidSet.SetNumUninitialized(Selection.SelectedIDs.Num() * 3);
+			TidSet.SetNumUninitialized(Selection.SelectedIDs.Num());
+
+			int32 TIndex = 0;
+			for (const int32 Tid : Selection.SelectedIDs)
 			{
-				FIndex3i TriVids = UnwrapMesh->GetTriangle(Tid);
-				for (int i = 0; i < 3; ++i)
-				{
-					VidSet.Add(TriVids[i]);
-				}
-				TidSet.Add(Tid);
+				const FIndex3i TriVids = UnwrapMesh->GetTriangle(Tid);
+				VidSet[TIndex * 3 + 0] = TriVids[0];
+				VidSet[TIndex * 3 + 1] = TriVids[1];
+				VidSet[TIndex * 3 + 2] = TriVids[2];
+				TidSet[TIndex++] = Tid;
 			}
 		}
 		else if (Selection.Type == FUVToolSelection::EType::Edge)
 		{
 			TRACE_CPUPROFILER_EVENT_SCOPE(Edge);
 
-			for (int32 Eid : Selection.SelectedIDs)
+			VidSet.SetNumUninitialized(Selection.SelectedIDs.Num() * 2);
+			TidSet.Reserve(Selection.SelectedIDs.Num());
+
+			int32 VIndex = 0;
+			TArray<int> Tids;
+			for (const int32 Eid : Selection.SelectedIDs)
 			{
-				FIndex2i EdgeVids = UnwrapMesh->GetEdgeV(Eid);
-				for (int i = 0; i < 2; ++i)
-				{
-					VidSet.Add(EdgeVids[i]);
-
-					TArray<int> TidOneRing;
-					UnwrapMesh->GetVtxTriangles(EdgeVids[i], TidOneRing);
-					TidSet.Append(TidOneRing);
-				}
+				const FIndex2i EdgeVids = UnwrapMesh->GetEdgeV(Eid);
+				VidSet[VIndex++] = EdgeVids[0];
+				VidSet[VIndex++] = EdgeVids[1];
+
+				Tids.Reset();
+				UnwrapMesh->GetVtxTriangles(EdgeVids[0], Tids);
+				UnwrapMesh->GetVtxTriangles(EdgeVids[1], Tids);
+				TidSet.Append(Tids);
 			}
 		}
 		else if (Selection.Type == FUVToolSelection::EType::Vertex)
 		{
 			TRACE_CPUPROFILER_EVENT_SCOPE(Vertex);
 
-			for (int32 Vid : Selection.SelectedIDs)
+			VidSet.SetNumUninitialized(Selection.SelectedIDs.Num());
+			TidSet.Reserve(Selection.SelectedIDs.Num());
+			
+			int32 VIndex = 0;
+			TArray<int> Tids;
+			for (const int32 Vid : Selection.SelectedIDs)
 			{
-				VidSet.Add(Vid);
+				VidSet[VIndex++] = Vid;
 
-				TArray<int> TidOneRing;
-				UnwrapMesh->GetVtxTriangles(Vid, TidOneRing);
-				TidSet.Append(TidOneRing);
+				Tids.Reset();
+				UnwrapMesh->GetVtxTriangles(Vid, Tids);
+				TidSet.Append(Tids);
 			}
 		}
 		else
@@ -439,8 +452,14 @@ void UUVSelectTool::OnSelectionChanged(bool, uint32 SelectionChangeType)
 			ensure(false);
 		}
 
-		MovingVidsPerSelection.Emplace(VidSet.Array());
-		RenderUpdateTidsPerSelection.Emplace(TidSet.Array());
+		// Make entries unique
+		VidSet.Sort();
+		VidSet.SetNum(Algo::Unique(VidSet));
+		TidSet.Sort();
+		TidSet.SetNum(Algo::Unique(TidSet));
+		
+		MovingVidsPerSelection.Emplace(VidSet);
+		RenderUpdateTidsPerSelection.Emplace(TidSet);
 	}
 
 	UpdateGizmo();
diff --git a/Engine/Source/Runtime/GeometryCore/Private/DynamicMesh/DynamicMesh3_Queries.cpp b/Engine/Source/Runtime/GeometryCore/Private/DynamicMesh/DynamicMesh3_Queries.cpp
--- a/Engine/Source/Runtime/GeometryCore/Private/DynamicMesh/DynamicMesh3_Queries.cpp
+++ b/Engine/Source/Runtime/GeometryCore/Private/DynamicMesh/DynamicMesh3_Queries.cpp
@@ -147,7 +147,6 @@ EMeshResult FDynamicMesh3::GetVtxTriangles(int vID, TArray<int>& TrianglesOut) c
 			{
 				TrianglesOut.Add(Edge.Tri[0]);
 			}
-			const int et1 = Edge.Tri[1];
 			if (Edge.Tri[1] != InvalidID && TriHasSequentialVertices(Edge.Tri[1], vID, vOther))
 			{
 				TrianglesOut.Add(Edge.Tri[1]);
