UVEditor: Minor changes to tool property labeling to make it more consistent with ModelingMode equivalent tools.

diff --git a/Engine/Plugins/Experimental/UVEditor/Source/UVEditorTools/Public/Operators/UVEditorRecomputeUVsOp.h b/Engine/Plugins/Experimental/UVEditor/Source/UVEditorTools/Public/Operators/UVEditorRecomputeUVsOp.h
--- a/Engine/Plugins/Experimental/UVEditor/Source/UVEditorTools/Public/Operators/UVEditorRecomputeUVsOp.h
+++ b/Engine/Plugins/Experimental/UVEditor/Source/UVEditorTools/Public/Operators/UVEditorRecomputeUVsOp.h
@@ -89,7 +89,7 @@ public:
 
 	/** Type of UV flattening algorithm to use */
 	UPROPERTY(EditAnywhere, Category = "UV Unwrap")
-	EUVEditorRecomputeUVsPropertiesUnwrapType UnwrapType = EUVEditorRecomputeUVsPropertiesUnwrapType::Conformal;
+	EUVEditorRecomputeUVsPropertiesUnwrapType UnwrapType = EUVEditorRecomputeUVsPropertiesUnwrapType::SpectralConformal;
 
 	/** Type of automatic rotation applied to each UV island */
 	UPROPERTY(EditAnywhere, Category = "UV Unwrap")
diff --git a/Engine/Plugins/Experimental/UVEditor/Source/UVEditorToolsEditorOnly/Public/Operators/UVEditorParameterizeMeshOp.h b/Engine/Plugins/Experimental/UVEditor/Source/UVEditorToolsEditorOnly/Public/Operators/UVEditorParameterizeMeshOp.h
--- a/Engine/Plugins/Experimental/UVEditor/Source/UVEditorToolsEditorOnly/Public/Operators/UVEditorParameterizeMeshOp.h
+++ b/Engine/Plugins/Experimental/UVEditor/Source/UVEditorToolsEditorOnly/Public/Operators/UVEditorParameterizeMeshOp.h
@@ -57,7 +57,7 @@ public:
 	int TextureResolution = 1024;
 
 	/** Generate new UVs based on polygroups from specified layer. */
-	UPROPERTY(EditAnywhere, Category = UVAtlas, meta = (EditCondition = bPolygroupsEnabled, EditConditionHides, HideEditConditionToggle))
+	UPROPERTY(EditAnywhere, Category = UVAtlas, meta = (DisplayName = "Constrain to Polygroups", EditCondition = bPolygroupsEnabled, EditConditionHides, HideEditConditionToggle))
 	bool bUsePolygroups;
 
 	/** Layout resulting islands on UDIMs based on polygroups. */
@@ -123,7 +123,7 @@ public:
 	int TextureResolution = 1024;
 
 	/** Generate new UVs based on polygroups from specified layer. */
-	UPROPERTY(EditAnywhere, Category = PatchBuilder, meta = (EditCondition = bPolygroupsEnabled, EditConditionHides, HideEditConditionToggle))
+	UPROPERTY(EditAnywhere, Category = PatchBuilder, meta = (DisplayName = "Constrain to Polygroups", EditCondition = bPolygroupsEnabled, EditConditionHides, HideEditConditionToggle))
 	bool bUsePolygroups;
 
 	/** Layout resulting islands on UDIMs based on polygroups. */
