UVEditor: Change the Selection Gizmo placement from being the centroid to the bounding box center of the selected elements.

diff --git a/Engine/Plugins/Experimental/UVEditor/Source/UVEditorTools/Private/Selection/UVToolSelectionAPI.cpp b/Engine/Plugins/Experimental/UVEditor/Source/UVEditorTools/Private/Selection/UVToolSelectionAPI.cpp
--- a/Engine/Plugins/Experimental/UVEditor/Source/UVEditorTools/Private/Selection/UVToolSelectionAPI.cpp
+++ b/Engine/Plugins/Experimental/UVEditor/Source/UVEditorTools/Private/Selection/UVToolSelectionAPI.cpp
@@ -151,7 +151,7 @@ void UUVToolSelectionAPI::SetSelections(const TArray<FUVToolSelection>& Selectio
 	bool bWillCallEndChange = bEmitChange && !bUsingExistingTransaction;
 
 
-	bCachedUnwrapSelectionCentroidValid = false;
+	bCachedUnwrapSelectionBoundingBoxCenterValid = false;
 
 	// If we don't match with the current unset selection type, just clear them to keep everything consistent
 	if (CurrentUnsetSelections.Num() > 0 && SelectionsIn.Num() > 0 && SelectionsIn[0].Type != CurrentUnsetSelections[0].Type)
@@ -209,7 +209,7 @@ void UUVToolSelectionAPI::SetSelections(const TArray<FUVToolSelection>& Selectio
 			FTransform Transform = FTransform::Identity;
 			if (HighlightOptions.bUseCentroidForUnwrapAutoUpdate)
 			{
-				Transform = FTransform(GetUnwrapSelectionCentroid());
+				Transform = FTransform(GetUnwrapSelectionBoundingBoxCenter());
 			}
 			RebuildUnwrapHighlight(Transform);
 		}
@@ -236,56 +236,26 @@ void UUVToolSelectionAPI::SetSelectionMechanicMode(EUVEditorSelectionMode Mode,
 	SelectionMechanic->SetSelectionMode(Mode, Options);
 }
 
-FVector3d UUVToolSelectionAPI::GetUnwrapSelectionCentroid(bool bForceRecalculate)
+FVector3d UUVToolSelectionAPI::GetUnwrapSelectionBoundingBoxCenter(bool bForceRecalculate)
 {
-	if (bCachedUnwrapSelectionCentroidValid && !bForceRecalculate)
+	if (bCachedUnwrapSelectionBoundingBoxCenterValid && !bForceRecalculate)
 	{
-		return CachedUnwrapSelectionCentroid;
+		return CachedUnwrapSelectionBoundingBoxCenter;
 	}
 
-	FVector3d& Centroid = CachedUnwrapSelectionCentroid;
-
-	Centroid = FVector3d::Zero();
+	FVector3d& Center = CachedUnwrapSelectionBoundingBoxCenter;
+	
+	FAxisAlignedBox3d AllSelectionsBoundingBox;
 	double Divisor = 0;
 	for (const FUVToolSelection& Selection : GetSelections())
 	{
 		FDynamicMesh3* Mesh = Selection.Target->UnwrapCanonical.Get();
-		if (Selection.Type == FUVToolSelection::EType::Edge)
-		{
-			for (int32 Eid : Selection.SelectedIDs)
-			{
-				Centroid += Mesh->GetEdgePoint(Eid, 0.5);
-			}
-		}
-		else if (Selection.Type == FUVToolSelection::EType::Triangle)
-		{
-			for (int32 Tid : Selection.SelectedIDs)
-			{
-				Centroid += Mesh->GetTriCentroid(Tid);
-			}
-		}
-		else if (Selection.Type == FUVToolSelection::EType::Vertex)
-		{
-			for (int32 Vid : Selection.SelectedIDs)
-			{
-				Centroid += Mesh->GetVertex(Vid);
-			}
-		}
-		else
-		{
-			ensure(false);
-		}
-
-		Divisor += Selection.SelectedIDs.Num();
-	}
-
-	if (Divisor > 0)
-	{
-		Centroid /= Divisor;
+		AllSelectionsBoundingBox.Contain(Selection.ToBoundingBox(*Mesh));
 	}
+	Center = AllSelectionsBoundingBox.Center();
 
-	bCachedUnwrapSelectionCentroidValid = true;
-	return Centroid;
+	bCachedUnwrapSelectionBoundingBoxCenterValid = true;
+	return Center;
 }
 
 
@@ -375,7 +345,7 @@ void UUVToolSelectionAPI::SetHighlightVisible(bool bUnwrapHighlightVisible, bool
 			FTransform UnwrapTransform = FTransform::Identity;
 			if (HighlightOptions.bUseCentroidForUnwrapAutoUpdate)
 			{
-				UnwrapTransform = FTransform(GetUnwrapSelectionCentroid());
+				UnwrapTransform = FTransform(GetUnwrapSelectionBoundingBoxCenter());
 			}
 			RebuildUnwrapHighlight(UnwrapTransform);
 		}
diff --git a/Engine/Plugins/Experimental/UVEditor/Source/UVEditorTools/Private/UVSelectTool.cpp b/Engine/Plugins/Experimental/UVEditor/Source/UVEditorTools/Private/UVSelectTool.cpp
--- a/Engine/Plugins/Experimental/UVEditor/Source/UVEditorTools/Private/UVSelectTool.cpp
+++ b/Engine/Plugins/Experimental/UVEditor/Source/UVEditorTools/Private/UVSelectTool.cpp
@@ -341,7 +341,7 @@ void UUVSelectTool::UpdateGizmo()
 {
 	if (SelectionAPI->HaveSelections())
 	{
-		TransformGizmo->ReinitializeGizmoTransform(FTransform(SelectionAPI->GetUnwrapSelectionCentroid()));
+		TransformGizmo->ReinitializeGizmoTransform(FTransform(SelectionAPI->GetUnwrapSelectionBoundingBoxCenter()));
 	}
 
 	TransformGizmo->SetVisibility(
@@ -434,7 +434,6 @@ void UUVSelectTool::OnSelectionChanged(bool, uint32 SelectionChangeType)
 			TRACE_CPUPROFILER_EVENT_SCOPE(Vertex);
 
 			VidSet.SetNumUninitialized(Selection.SelectedIDs.Num());
-			TidSet.Reserve(Selection.SelectedIDs.Num());
 			
 			int32 VIndex = 0;
 			TArray<int> Tids;
diff --git a/Engine/Plugins/Experimental/UVEditor/Source/UVEditorTools/Public/Selection/UVToolSelectionAPI.h b/Engine/Plugins/Experimental/UVEditor/Source/UVEditorTools/Public/Selection/UVToolSelectionAPI.h
--- a/Engine/Plugins/Experimental/UVEditor/Source/UVEditorTools/Public/Selection/UVToolSelectionAPI.h
+++ b/Engine/Plugins/Experimental/UVEditor/Source/UVEditorTools/Public/Selection/UVToolSelectionAPI.h
@@ -94,11 +94,11 @@ public:
 	void ClearSelections(bool bBroadcast, bool bEmitChange);
 
 	/**
-	 * Gets the centroid of the current selection. The centroid is cached and
+	 * Gets the bounding box center of the current selection. The center point is cached and
 	 * invalidated on the next SetSelection call, but can be forced to be
 	 * recalculated.
 	 */
-	FVector3d GetUnwrapSelectionCentroid(bool bForceRecalculate = false);
+	FVector3d GetUnwrapSelectionBoundingBoxCenter(bool bForceRecalculate = false);
 
 	// These are functions that are used if a selection made in the applied mesh cannot be
 	// converted to an unwrap mesh selection due to unset UV elements (meaning that the selected
@@ -357,8 +357,8 @@ protected:
 
 	FHighlightOptions HighlightOptions;
 
-	FVector3d CachedUnwrapSelectionCentroid;
-	bool bCachedUnwrapSelectionCentroidValid = false;
+	FVector3d CachedUnwrapSelectionBoundingBoxCenter;
+	bool bCachedUnwrapSelectionBoundingBoxCenterValid = false;
 };
 
 // UInterface for IUVToolSupportsSelection
