UVEditor: Add new high performance rendering components for handling mesh highlighting on large meshes. Includes 2D and 3D variations of PointSets, LineSets, and TriangleSets.

These are considerably stripped down from the previously used component classes, designed with the intent on reducing needed memory bandwidth to move data to the rendering thread.

diff --git a/Engine/Plugins/Experimental/UVEditor/Source/UVEditorTools/Private/Drawing/BasicLineSetComponent.cpp b/Engine/Plugins/Experimental/UVEditor/Source/UVEditorTools/Private/Drawing/BasicLineSetComponent.cpp
new file mode 100644
--- /dev/null
+++ b/Engine/Plugins/Experimental/UVEditor/Source/UVEditorTools/Private/Drawing/BasicLineSetComponent.cpp
@@ -0,0 +1,251 @@
+// Copyright Epic Games, Inc. All Rights Reserved.
+
+#include "Drawing/BasicLineSetComponent.h"
+#include "PrimitiveSceneProxy.h"
+#include "Materials/Material.h"
+#include "DynamicMeshBuilder.h"
+#include "StaticMeshResources.h"
+
+
+struct FBasicLineSetMeshBatchData
+{
+	FMaterialRenderProxy* MaterialProxy = nullptr;
+	int32 StartIndex = -1;
+	int32 NumPrimitives = -1;
+	int32 MinVertexIndex = -1;
+	int32 MaxVertexIndex = -1;
+};
+
+class FBasicLineSetSceneProxy final : public FPrimitiveSceneProxy
+{
+public:
+
+	template <typename BasicLineSetComponent>
+	FBasicLineSetSceneProxy(BasicLineSetComponent* Component)
+		: FPrimitiveSceneProxy(Component),
+		MaterialRelevance(Component->GetMaterialRelevance(GetScene().GetFeatureLevel())),
+		VertexFactory(GetScene().GetFeatureLevel(), "FLineSetSceneProxy")
+	{
+		const int32 NumLineVertices = Component->NumElements() * 4;
+		const int32 NumLineIndices = Component->NumElements() * 6;
+		const int32 TotalNumVertices = NumLineVertices;
+		const int32 TotalNumIndices = NumLineIndices;
+		constexpr int32 NumTextureCoordinates = 1;
+
+		VertexBuffers.PositionVertexBuffer.Init(TotalNumVertices);
+		VertexBuffers.StaticMeshVertexBuffer.Init(TotalNumVertices, NumTextureCoordinates);
+		VertexBuffers.ColorVertexBuffer.Init(TotalNumVertices);
+		IndexBuffer.Indices.SetNumUninitialized(TotalNumIndices);
+
+		// Initialize points.
+		// Lines are represented as two tris, all of whose vertices are coincident.
+		// The material then offsets them according to the signs of the vertex normals in a camera facing orientation.
+		// Size of the point is given by U0.
+		if (Component->NumElements() > 0)
+		{
+			MeshBatchDatas.Emplace();
+			FBasicLineSetMeshBatchData& MeshBatchData = MeshBatchDatas.Last();
+			MeshBatchData.MinVertexIndex = 0;
+			MeshBatchData.MaxVertexIndex = NumLineVertices - 1;
+			MeshBatchData.StartIndex = 0;
+			MeshBatchData.NumPrimitives = Component->NumElements() * 2;
+			if (Component->GetMaterial(0) != nullptr)
+			{
+				MeshBatchData.MaterialProxy = Component->GetMaterial(0)->GetRenderProxy();
+			}
+			else
+			{
+				MeshBatchData.MaterialProxy = UMaterial::GetDefaultMaterial(MD_Surface)->GetRenderProxy();
+			}
+
+			const TArray<typename BasicLineSetComponent::GeometryType>& Lines = Component->Elements;
+
+			// The color stored in the vertices actually gets interpreted as a linear color by the material,
+			// whereas it is more convenient for the user of the LineSet to specify colors as sRGB. So we actually
+			// have to convert it back to linear. The ToFColor(false) call just scales back into 0-255 space.
+			const FColor Color = FLinearColor::FromSRGBColor(Component->Color).ToFColor(false);
+
+			const FVector2f UV(Component->Size, Component->DepthBias);
+
+			ParallelFor(Component->NumElements(), [&](int32 i)
+				{
+					const int32 VertexBufferIndex = 4 * i;
+					const int32 IndexBufferIndex = 6 * i;
+					const int32 LineSize = BasicLineSetComponent::ElementSize;
+					
+					const FVector3f LineDirection = ToFVector3f(Lines[i * LineSize + 1] - Lines[i * LineSize + 0]).GetSafeNormal();
+
+					VertexBuffers.PositionVertexBuffer.VertexPosition(VertexBufferIndex + 0) = ToFVector3f(Lines[i * LineSize + 0]);
+					VertexBuffers.PositionVertexBuffer.VertexPosition(VertexBufferIndex + 1) = ToFVector3f(Lines[i * LineSize + 1]);
+					VertexBuffers.PositionVertexBuffer.VertexPosition(VertexBufferIndex + 2) = ToFVector3f(Lines[i * LineSize + 1]);
+					VertexBuffers.PositionVertexBuffer.VertexPosition(VertexBufferIndex + 3) = ToFVector3f(Lines[i * LineSize + 0]);
+
+					VertexBuffers.StaticMeshVertexBuffer.SetVertexTangents(VertexBufferIndex + 0, FVector3f::ZeroVector, FVector3f::ZeroVector, (FVector3f)-LineDirection);
+					VertexBuffers.StaticMeshVertexBuffer.SetVertexTangents(VertexBufferIndex + 1, FVector3f::ZeroVector, FVector3f::ZeroVector, (FVector3f)-LineDirection);
+					VertexBuffers.StaticMeshVertexBuffer.SetVertexTangents(VertexBufferIndex + 2, FVector3f::ZeroVector, FVector3f::ZeroVector, (FVector3f)LineDirection);
+					VertexBuffers.StaticMeshVertexBuffer.SetVertexTangents(VertexBufferIndex + 3, FVector3f::ZeroVector, FVector3f::ZeroVector, (FVector3f)LineDirection);
+
+					VertexBuffers.StaticMeshVertexBuffer.SetVertexUV(VertexBufferIndex + 0, 0, UV);
+					VertexBuffers.StaticMeshVertexBuffer.SetVertexUV(VertexBufferIndex + 1, 0, UV);
+					VertexBuffers.StaticMeshVertexBuffer.SetVertexUV(VertexBufferIndex + 2, 0, UV);
+					VertexBuffers.StaticMeshVertexBuffer.SetVertexUV(VertexBufferIndex + 3, 0, UV);
+
+					VertexBuffers.ColorVertexBuffer.VertexColor(VertexBufferIndex + 0) = Color;
+					VertexBuffers.ColorVertexBuffer.VertexColor(VertexBufferIndex + 1) = Color;
+					VertexBuffers.ColorVertexBuffer.VertexColor(VertexBufferIndex + 2) = Color;
+					VertexBuffers.ColorVertexBuffer.VertexColor(VertexBufferIndex + 3) = Color;
+
+					IndexBuffer.Indices[IndexBufferIndex + 0] = VertexBufferIndex + 0;
+					IndexBuffer.Indices[IndexBufferIndex + 1] = VertexBufferIndex + 1;
+					IndexBuffer.Indices[IndexBufferIndex + 2] = VertexBufferIndex + 2;
+					IndexBuffer.Indices[IndexBufferIndex + 3] = VertexBufferIndex + 2;
+					IndexBuffer.Indices[IndexBufferIndex + 4] = VertexBufferIndex + 3;
+					IndexBuffer.Indices[IndexBufferIndex + 5] = VertexBufferIndex + 0;
+				});
+
+			ENQUEUE_RENDER_COMMAND(OverlayVertexBuffersInit)(
+				[this](FRHICommandListImmediate& RHICmdList)
+				{
+					VertexBuffers.PositionVertexBuffer.InitResource();
+					VertexBuffers.StaticMeshVertexBuffer.InitResource();
+					VertexBuffers.ColorVertexBuffer.InitResource();
+
+					FLocalVertexFactory::FDataType Data;
+					VertexBuffers.PositionVertexBuffer.BindPositionVertexBuffer(&VertexFactory, Data);
+					VertexBuffers.StaticMeshVertexBuffer.BindTangentVertexBuffer(&VertexFactory, Data);
+					VertexBuffers.StaticMeshVertexBuffer.BindTexCoordVertexBuffer(&VertexFactory, Data);
+					VertexBuffers.ColorVertexBuffer.BindColorVertexBuffer(&VertexFactory, Data);
+					VertexFactory.SetData(Data);
+
+					VertexFactory.InitResource();
+					IndexBuffer.InitResource();
+				});
+		}
+	}
+
+	virtual ~FBasicLineSetSceneProxy() override
+	{
+		VertexBuffers.PositionVertexBuffer.ReleaseResource();
+		VertexBuffers.StaticMeshVertexBuffer.ReleaseResource();
+		VertexBuffers.ColorVertexBuffer.ReleaseResource();
+		IndexBuffer.ReleaseResource();
+		VertexFactory.ReleaseResource();
+	}
+
+	virtual void GetDynamicMeshElements(const TArray<const FSceneView*>& Views, const FSceneViewFamily& ViewFamily, uint32 VisibilityMap, FMeshElementCollector& Collector) const override
+	{
+		QUICK_SCOPE_CYCLE_COUNTER(STAT_OverlaySceneProxy_GetDynamicMeshElements);
+
+		for (int32 ViewIndex = 0; ViewIndex < Views.Num(); ViewIndex++)
+		{
+			if (VisibilityMap & (1 << ViewIndex))
+			{
+				for (const FBasicLineSetMeshBatchData& MeshBatchData : MeshBatchDatas)
+				{
+					FMeshBatch& Mesh = Collector.AllocateMesh();
+					FMeshBatchElement& BatchElement = Mesh.Elements[0];
+					BatchElement.IndexBuffer = &IndexBuffer;
+					Mesh.bWireframe = false;
+					Mesh.VertexFactory = &VertexFactory;
+					Mesh.MaterialRenderProxy = MeshBatchData.MaterialProxy;
+
+					FDynamicPrimitiveUniformBuffer& DynamicPrimitiveUniformBuffer = Collector.AllocateOneFrameResource<FDynamicPrimitiveUniformBuffer>();
+					DynamicPrimitiveUniformBuffer.Set(GetLocalToWorld(), GetLocalToWorld(), GetBounds(), GetLocalBounds(), true, false, AlwaysHasVelocity());
+					BatchElement.PrimitiveUniformBufferResource = &DynamicPrimitiveUniformBuffer.UniformBuffer;
+
+					BatchElement.FirstIndex = MeshBatchData.StartIndex;
+					BatchElement.NumPrimitives = MeshBatchData.NumPrimitives;
+					BatchElement.MinVertexIndex = MeshBatchData.MinVertexIndex;
+					BatchElement.MaxVertexIndex = MeshBatchData.MaxVertexIndex;
+					Mesh.ReverseCulling = IsLocalToWorldDeterminantNegative();
+					Mesh.Type = PT_TriangleList;
+					Mesh.DepthPriorityGroup = SDPG_World;
+					Mesh.bCanApplyViewModeOverrides = false;
+					Collector.AddMesh(ViewIndex, Mesh);
+				}
+			}
+		}
+	}
+
+	virtual FPrimitiveViewRelevance GetViewRelevance(const FSceneView* View) const override
+	{
+		FPrimitiveViewRelevance Result;
+		Result.bDrawRelevance = IsShown(View);
+		Result.bShadowRelevance = IsShadowCast(View);
+		Result.bDynamicRelevance = true;
+		Result.bRenderInMainPass = ShouldRenderInMainPass();
+		Result.bUsesLightingChannels = GetLightingChannelMask() != GetDefaultLightingChannelMask();
+		Result.bRenderCustomDepth = ShouldRenderCustomDepth();
+		Result.bTranslucentSelfShadow = bCastVolumetricTranslucentShadow;
+		MaterialRelevance.SetPrimitiveViewRelevance(Result);
+		Result.bVelocityRelevance = DrawsVelocity() && Result.bOpaque && Result.bRenderInMainPass;
+		return Result;
+	}
+
+	virtual bool CanBeOccluded() const override
+	{
+		return !MaterialRelevance.bDisableDepthTest;
+	}
+
+	virtual uint32 GetMemoryFootprint() const override { return sizeof(*this) + GetAllocatedSize(); }
+
+	virtual SIZE_T GetTypeHash() const override
+	{
+		static SIZE_T UniquePointer;
+		return reinterpret_cast<SIZE_T>(&UniquePointer);
+	}
+
+private:
+
+	// Line type specific implementation to get an FVector3f out of the generic point type.
+	static FVector3f ToFVector3f(const FVector2f& Vector) { return FVector3f(Vector.X, Vector.Y, 0.0f); }
+	static FVector3f ToFVector3f(const FVector3f& Vector) { return Vector; }
+
+	TArray<FBasicLineSetMeshBatchData> MeshBatchDatas;
+	FMaterialRelevance MaterialRelevance;
+	FLocalVertexFactory VertexFactory;
+	FStaticMeshVertexBuffers VertexBuffers;
+	FDynamicMeshIndexBuffer32 IndexBuffer;
+};
+
+
+UBasicLineSetComponentBase::UBasicLineSetComponentBase()
+{
+	CastShadow = false;
+	bSelectable = false;
+	PrimaryComponentTick.bCanEverTick = false;
+
+	UPrimitiveComponent::SetCollisionProfileName(UCollisionProfile::NoCollision_ProfileName);
+}
+
+void UBasicLineSetComponentBase::ClearComponent()
+{
+	MarkRenderStateDirty();
+	bBoundsDirty = true;
+}
+
+
+void UBasicLineSetComponentBase::SetLineMaterial(UMaterialInterface* InLineMaterial)
+{
+	LineMaterial = InLineMaterial;
+	UMeshComponent::SetMaterial(0, InLineMaterial);
+}
+
+void UBasicLineSetComponentBase::SetLineSetParameters(FColor InColor, float InSize, float InDepthBias)
+{
+	Color = InColor;
+	Size = InSize;
+	DepthBias = InDepthBias;
+	bBoundsDirty = true;
+}
+
+
+FPrimitiveSceneProxy* UBasic2DLineSetComponent::CreateSceneProxy()
+{
+	return Elements.Num() > 0 ? new FBasicLineSetSceneProxy(this) : nullptr;
+}
+
+FPrimitiveSceneProxy* UBasic3DLineSetComponent::CreateSceneProxy()
+{
+	return Elements.Num() > 0 ? new FBasicLineSetSceneProxy(this) : nullptr;
+}
diff --git a/Engine/Plugins/Experimental/UVEditor/Source/UVEditorTools/Private/Drawing/BasicPointSetComponent.cpp b/Engine/Plugins/Experimental/UVEditor/Source/UVEditorTools/Private/Drawing/BasicPointSetComponent.cpp
new file mode 100644
--- /dev/null
+++ b/Engine/Plugins/Experimental/UVEditor/Source/UVEditorTools/Private/Drawing/BasicPointSetComponent.cpp
@@ -0,0 +1,242 @@
+// Copyright Epic Games, Inc. All Rights Reserved.
+
+#include "Drawing/BasicPointSetComponent.h"
+#include "PrimitiveSceneProxy.h"
+#include "Materials/Material.h"
+#include "DynamicMeshBuilder.h"
+#include "StaticMeshResources.h"
+
+
+struct FBasicPointSetMeshBatchData
+{
+	FMaterialRenderProxy* MaterialProxy = nullptr;
+	int32 StartIndex = -1;
+	int32 NumPrimitives = -1;
+	int32 MinVertexIndex = -1;
+	int32 MaxVertexIndex = -1;
+};
+
+class FBasicPointSetSceneProxy final : public FPrimitiveSceneProxy
+{
+public:
+
+	template <typename BasicPointSetComponent>
+	FBasicPointSetSceneProxy(BasicPointSetComponent* Component)
+		: FPrimitiveSceneProxy(Component),
+		MaterialRelevance(Component->GetMaterialRelevance(GetScene().GetFeatureLevel())),
+		VertexFactory(GetScene().GetFeatureLevel(), "FPointSetSceneProxy")
+	{
+		const int32 NumPointVertices = Component->NumElements() * 4;
+		const int32 NumPointIndices = Component->NumElements() * 6;
+		const int32 TotalNumVertices = NumPointVertices;
+		const int32 TotalNumIndices = NumPointIndices;
+		constexpr int32 NumTextureCoordinates = 1;
+
+		VertexBuffers.PositionVertexBuffer.Init(TotalNumVertices);
+		VertexBuffers.StaticMeshVertexBuffer.Init(TotalNumVertices, NumTextureCoordinates);
+		VertexBuffers.ColorVertexBuffer.Init(TotalNumVertices);
+		IndexBuffer.Indices.SetNumUninitialized(TotalNumIndices);
+			
+		// Initialize points.
+		// Points are represented as two tris, all of whose vertices are coincident.
+		// The material then offsets them according to the signs of the vertex normals in a camera facing orientation.
+		// Size of the point is given by U0.
+		if (Component->NumElements() > 0)
+		{
+			MeshBatchDatas.Emplace();
+			FBasicPointSetMeshBatchData& MeshBatchData = MeshBatchDatas.Last();
+			MeshBatchData.MinVertexIndex = 0;
+			MeshBatchData.MaxVertexIndex = NumPointVertices - 1;
+			MeshBatchData.StartIndex = 0;
+			MeshBatchData.NumPrimitives = Component->NumElements() * 2;
+			if (Component->GetMaterial(0) != nullptr)
+			{
+				MeshBatchData.MaterialProxy = Component->GetMaterial(0)->GetRenderProxy();
+			}
+			else
+			{
+				MeshBatchData.MaterialProxy = UMaterial::GetDefaultMaterial(MD_Surface)->GetRenderProxy();
+			}
+
+			const FVector TangentVectors[4] = {
+				FVector(1.0f, -1.0f, 0.0f),
+				FVector(1.0f, 1.0f, 0.0f),
+				FVector(-1.0f, 1.0f, 0.0f),
+				FVector(-1.0f, -1.0f, 0.0f)
+			};
+
+			const TArray<typename BasicPointSetComponent::GeometryType>& Points = Component->Elements;
+
+			// The color stored in the vertices actually gets interpreted as a linear color by the material,
+			// whereas it is more convenient for the user of the LineSet to specify colors as sRGB. So we actually
+			// have to convert it back to linear. The ToFColor(false) call just scales back into 0-255 space.
+			const FColor Color = FLinearColor::FromSRGBColor(Component->Color).ToFColor(false);
+
+			const FVector2f UV(Component->Size, Component->DepthBias);
+
+			ParallelFor(Component->NumElements(), [&](int32 i)
+			{
+				const int32 VertexBufferIndex = 4 * i;
+				const int32 IndexBufferIndex = 6 * i;
+
+				for (int j = 0; j < 4; ++j)
+				{
+					VertexBuffers.PositionVertexBuffer.VertexPosition(VertexBufferIndex + j) = ToFVector3f(Points[i]);
+					VertexBuffers.StaticMeshVertexBuffer.SetVertexUV(VertexBufferIndex + j, 0, UV);
+					VertexBuffers.ColorVertexBuffer.VertexColor(VertexBufferIndex + j) = Color;
+					VertexBuffers.StaticMeshVertexBuffer.SetVertexTangents(VertexBufferIndex + j, FVector3f::ZeroVector, FVector3f::ZeroVector, static_cast<FVector3f>(TangentVectors[j]));
+				}
+
+				IndexBuffer.Indices[IndexBufferIndex + 0] = VertexBufferIndex + 0;
+				IndexBuffer.Indices[IndexBufferIndex + 1] = VertexBufferIndex + 1;
+				IndexBuffer.Indices[IndexBufferIndex + 2] = VertexBufferIndex + 2;
+				IndexBuffer.Indices[IndexBufferIndex + 3] = VertexBufferIndex + 2;
+				IndexBuffer.Indices[IndexBufferIndex + 4] = VertexBufferIndex + 3;
+				IndexBuffer.Indices[IndexBufferIndex + 5] = VertexBufferIndex + 0;
+			});
+		}
+
+		ENQUEUE_RENDER_COMMAND(OverlayVertexBuffersInit)(
+			[this](FRHICommandListImmediate& RHICmdList)
+		{
+			VertexBuffers.PositionVertexBuffer.InitResource();
+			VertexBuffers.StaticMeshVertexBuffer.InitResource();
+			VertexBuffers.ColorVertexBuffer.InitResource();
+
+			FLocalVertexFactory::FDataType Data;
+			VertexBuffers.PositionVertexBuffer.BindPositionVertexBuffer(&VertexFactory, Data);
+			VertexBuffers.StaticMeshVertexBuffer.BindTangentVertexBuffer(&VertexFactory, Data);
+			VertexBuffers.StaticMeshVertexBuffer.BindTexCoordVertexBuffer(&VertexFactory, Data);
+			VertexBuffers.ColorVertexBuffer.BindColorVertexBuffer(&VertexFactory, Data);
+			VertexFactory.SetData(Data);
+
+			VertexFactory.InitResource();
+			IndexBuffer.InitResource();
+		});
+	}
+
+	virtual ~FBasicPointSetSceneProxy() override
+	{
+		VertexBuffers.PositionVertexBuffer.ReleaseResource();
+		VertexBuffers.StaticMeshVertexBuffer.ReleaseResource();
+		VertexBuffers.ColorVertexBuffer.ReleaseResource();
+		IndexBuffer.ReleaseResource();
+		VertexFactory.ReleaseResource();
+	}
+
+	virtual void GetDynamicMeshElements(const TArray<const FSceneView*>& Views, const FSceneViewFamily& ViewFamily, uint32 VisibilityMap, FMeshElementCollector& Collector) const override
+	{
+		QUICK_SCOPE_CYCLE_COUNTER(STAT_OverlaySceneProxy_GetDynamicMeshElements);
+
+		for (int32 ViewIndex = 0; ViewIndex < Views.Num(); ViewIndex++)
+		{
+			if (VisibilityMap & (1 << ViewIndex))
+			{
+				for (const FBasicPointSetMeshBatchData& MeshBatchData : MeshBatchDatas)
+				{
+					FMeshBatch& Mesh = Collector.AllocateMesh();
+					FMeshBatchElement& BatchElement = Mesh.Elements[0];
+					BatchElement.IndexBuffer = &IndexBuffer;
+					Mesh.bWireframe = false;
+					Mesh.VertexFactory = &VertexFactory;
+					Mesh.MaterialRenderProxy = MeshBatchData.MaterialProxy;
+
+					FDynamicPrimitiveUniformBuffer& DynamicPrimitiveUniformBuffer = Collector.AllocateOneFrameResource<FDynamicPrimitiveUniformBuffer>();
+					DynamicPrimitiveUniformBuffer.Set(GetLocalToWorld(), GetLocalToWorld(), GetBounds(), GetLocalBounds(), true, false, AlwaysHasVelocity());
+					BatchElement.PrimitiveUniformBufferResource = &DynamicPrimitiveUniformBuffer.UniformBuffer;
+
+					BatchElement.FirstIndex = MeshBatchData.StartIndex;
+					BatchElement.NumPrimitives = MeshBatchData.NumPrimitives;
+					BatchElement.MinVertexIndex = MeshBatchData.MinVertexIndex;
+					BatchElement.MaxVertexIndex = MeshBatchData.MaxVertexIndex;
+					Mesh.ReverseCulling = IsLocalToWorldDeterminantNegative();
+					Mesh.Type = PT_TriangleList;
+					Mesh.DepthPriorityGroup = SDPG_World;
+					Mesh.bCanApplyViewModeOverrides = false;
+					Collector.AddMesh(ViewIndex, Mesh);
+				}
+			}
+		}
+	}
+
+	virtual FPrimitiveViewRelevance GetViewRelevance(const FSceneView* View) const override
+	{
+		FPrimitiveViewRelevance Result;
+		Result.bDrawRelevance = IsShown(View);
+		Result.bShadowRelevance = IsShadowCast(View);
+		Result.bDynamicRelevance = true;
+		Result.bRenderInMainPass = ShouldRenderInMainPass();
+		Result.bUsesLightingChannels = GetLightingChannelMask() != GetDefaultLightingChannelMask();
+		Result.bRenderCustomDepth = ShouldRenderCustomDepth();
+		Result.bTranslucentSelfShadow = bCastVolumetricTranslucentShadow;
+		MaterialRelevance.SetPrimitiveViewRelevance(Result);
+		Result.bVelocityRelevance = DrawsVelocity() && Result.bOpaque && Result.bRenderInMainPass;
+		return Result;
+	}
+
+	virtual bool CanBeOccluded() const override
+	{
+		return !MaterialRelevance.bDisableDepthTest;
+	}
+
+	virtual uint32 GetMemoryFootprint() const override { return sizeof(*this) + GetAllocatedSize(); }
+
+	virtual SIZE_T GetTypeHash() const override
+	{
+		static SIZE_T UniquePointer;
+		return reinterpret_cast<SIZE_T>(&UniquePointer);
+	}
+
+private:
+
+	// Point type specific implementation to get an FVector3f out of the generic point type.
+	static FVector3f ToFVector3f(const FVector2f& Vector) {	return FVector3f(Vector.X, Vector.Y, 0.0f); }
+	static FVector3f ToFVector3f(const FVector3f& Vector) {	return Vector; }
+	
+	TArray<FBasicPointSetMeshBatchData> MeshBatchDatas;
+	FMaterialRelevance MaterialRelevance;
+	FLocalVertexFactory VertexFactory;
+	FStaticMeshVertexBuffers VertexBuffers;
+	FDynamicMeshIndexBuffer32 IndexBuffer;
+};
+
+UBasicPointSetComponentBase::UBasicPointSetComponentBase()
+{
+	CastShadow = false;
+	bSelectable = false;
+	PrimaryComponentTick.bCanEverTick = false;
+
+	UPrimitiveComponent::SetCollisionProfileName(UCollisionProfile::NoCollision_ProfileName);
+}
+
+void UBasicPointSetComponentBase::ClearComponent()
+{
+	MarkRenderStateDirty();
+	bBoundsDirty = true;
+}
+
+
+void UBasicPointSetComponentBase::SetPointMaterial(UMaterialInterface* InPointMaterial)
+{
+	PointMaterial = InPointMaterial;
+	UMeshComponent::SetMaterial(0, InPointMaterial);
+}
+
+void UBasicPointSetComponentBase::SetPointSetParameters(FColor InColor, float InSize, float InDepthBias)
+{
+	Color = InColor;
+	Size = InSize;
+	DepthBias = InDepthBias;
+	bBoundsDirty = true;
+}
+
+
+FPrimitiveSceneProxy* UBasic2DPointSetComponent::CreateSceneProxy()
+{
+	return Elements.Num() > 0 ? new FBasicPointSetSceneProxy(this) : nullptr;
+}
+
+FPrimitiveSceneProxy* UBasic3DPointSetComponent::CreateSceneProxy()
+{
+	return Elements.Num() > 0 ? new FBasicPointSetSceneProxy(this) : nullptr;
+}
diff --git a/Engine/Plugins/Experimental/UVEditor/Source/UVEditorTools/Private/Drawing/BasicTriangleSetComponent.cpp b/Engine/Plugins/Experimental/UVEditor/Source/UVEditorTools/Private/Drawing/BasicTriangleSetComponent.cpp
new file mode 100644
--- /dev/null
+++ b/Engine/Plugins/Experimental/UVEditor/Source/UVEditorTools/Private/Drawing/BasicTriangleSetComponent.cpp
@@ -0,0 +1,239 @@
+// Copyright Epic Games, Inc. All Rights Reserved.
+
+#include "Drawing/BasicTriangleSetComponent.h"
+#include "PrimitiveSceneProxy.h"
+#include "Materials/Material.h"
+#include "DynamicMeshBuilder.h"
+#include "StaticMeshResources.h"
+
+
+struct FBasicTriangleSetMeshBatchData
+{
+	FMaterialRenderProxy* MaterialProxy = nullptr;
+	int32 StartIndex = -1;
+	int32 NumPrimitives = -1;
+	int32 MinVertexIndex = -1;
+	int32 MaxVertexIndex = -1;
+};
+
+class FBasicTriangleSetSceneProxy final : public FPrimitiveSceneProxy
+{
+public:
+
+	template <typename BasicTriangleSetComponent>
+	FBasicTriangleSetSceneProxy(BasicTriangleSetComponent* Component)
+		: FPrimitiveSceneProxy(Component),
+		MaterialRelevance(Component->GetMaterialRelevance(GetScene().GetFeatureLevel())),
+		VertexFactory(GetScene().GetFeatureLevel(), "FTriangleSetSceneProxy")
+	{
+		const int32 NumTriangleVertices = Component->NumElements() * 3;
+		const int32 NumTriangleIndices = Component->NumElements() * 3;
+		const int32 TotalNumVertices = NumTriangleVertices;
+		const int32 TotalNumIndices = NumTriangleIndices;
+		constexpr int32 NumTextureCoordinates = 1;
+
+		VertexBuffers.PositionVertexBuffer.Init(TotalNumVertices);
+		VertexBuffers.StaticMeshVertexBuffer.Init(TotalNumVertices, NumTextureCoordinates);
+		VertexBuffers.ColorVertexBuffer.Init(TotalNumVertices);
+		IndexBuffer.Indices.SetNumUninitialized(TotalNumIndices);
+
+		// Initialize points.
+		// Triangles are represented as two tris, all of whose vertices are coincident.
+		// The material then offsets them according to the signs of the vertex normals in a camera facing orientation.
+		// Size of the point is given by U0.
+		if (Component->NumElements() > 0)
+		{
+			MeshBatchDatas.Emplace();
+			FBasicTriangleSetMeshBatchData& MeshBatchData = MeshBatchDatas.Last();
+			MeshBatchData.MinVertexIndex = 0;
+			MeshBatchData.MaxVertexIndex = NumTriangleVertices - 1;
+			MeshBatchData.StartIndex = 0;
+			MeshBatchData.NumPrimitives = Component->NumElements();
+			if (Component->GetMaterial(0) != nullptr)
+			{
+				MeshBatchData.MaterialProxy = Component->GetMaterial(0)->GetRenderProxy();
+			}
+			else
+			{
+				MeshBatchData.MaterialProxy = UMaterial::GetDefaultMaterial(MD_Surface)->GetRenderProxy();
+			}
+
+			const TArray<typename BasicTriangleSetComponent::GeometryType>& Triangles = Component->Elements;
+
+			// The color stored in the vertices actually gets interpreted as a linear color by the material,
+			// whereas it is more convenient for the user of the TriangleSet to specify colors as sRGB. So we actually
+			// have to convert it back to linear. The ToFColor(false) call just scales back into 0-255 space.
+			const FColor Color = FLinearColor::FromSRGBColor(Component->Color).ToFColor(false);
+			const FVector3f Normal = Component->Normal;
+			ParallelFor(Component->NumElements(), [&](int32 i)
+				{
+					const int32 VertexBufferIndex = 3 * i;
+					const int32 IndexBufferIndex = 3 * i;
+					const int32 TriangleSize = BasicTriangleSetComponent::ElementSize;
+
+					VertexBuffers.PositionVertexBuffer.VertexPosition(VertexBufferIndex + 0) = ToFVector3f(Triangles[i * TriangleSize + 0]);
+					VertexBuffers.PositionVertexBuffer.VertexPosition(VertexBufferIndex + 1) = ToFVector3f(Triangles[i * TriangleSize + 1]);
+					VertexBuffers.PositionVertexBuffer.VertexPosition(VertexBufferIndex + 2) = ToFVector3f(Triangles[i * TriangleSize + 2]);
+
+					VertexBuffers.StaticMeshVertexBuffer.SetVertexTangents(VertexBufferIndex + 0, FVector3f(1, 0, 0), FVector3f(0, 1, 0), Normal);
+					VertexBuffers.StaticMeshVertexBuffer.SetVertexTangents(VertexBufferIndex + 1, FVector3f(1, 0, 0), FVector3f(0, 1, 0), Normal);
+					VertexBuffers.StaticMeshVertexBuffer.SetVertexTangents(VertexBufferIndex + 2, FVector3f(1, 0, 0), FVector3f(0, 1, 0), Normal);
+
+					VertexBuffers.StaticMeshVertexBuffer.SetVertexUV(VertexBufferIndex + 0, 0, FVector2f(0, 0));
+					VertexBuffers.StaticMeshVertexBuffer.SetVertexUV(VertexBufferIndex + 1, 0, FVector2f(1, 0));
+					VertexBuffers.StaticMeshVertexBuffer.SetVertexUV(VertexBufferIndex + 2, 0, FVector2f(0, 1));
+
+					VertexBuffers.ColorVertexBuffer.VertexColor(VertexBufferIndex + 0) = Color;
+					VertexBuffers.ColorVertexBuffer.VertexColor(VertexBufferIndex + 1) = Color;
+					VertexBuffers.ColorVertexBuffer.VertexColor(VertexBufferIndex + 2) = Color;
+
+					IndexBuffer.Indices[IndexBufferIndex + 0] = VertexBufferIndex + 0;
+					IndexBuffer.Indices[IndexBufferIndex + 1] = VertexBufferIndex + 1;
+					IndexBuffer.Indices[IndexBufferIndex + 2] = VertexBufferIndex + 2;
+				});
+
+			ENQUEUE_RENDER_COMMAND(OverlayVertexBuffersInit)(
+				[this](FRHICommandListImmediate& RHICmdList)
+				{
+					VertexBuffers.PositionVertexBuffer.InitResource();
+					VertexBuffers.StaticMeshVertexBuffer.InitResource();
+					VertexBuffers.ColorVertexBuffer.InitResource();
+
+					FLocalVertexFactory::FDataType Data;
+					VertexBuffers.PositionVertexBuffer.BindPositionVertexBuffer(&VertexFactory, Data);
+					VertexBuffers.StaticMeshVertexBuffer.BindTangentVertexBuffer(&VertexFactory, Data);
+					VertexBuffers.StaticMeshVertexBuffer.BindTexCoordVertexBuffer(&VertexFactory, Data);
+					VertexBuffers.ColorVertexBuffer.BindColorVertexBuffer(&VertexFactory, Data);
+					VertexFactory.SetData(Data);
+
+					VertexFactory.InitResource();
+					IndexBuffer.InitResource();
+				});
+		}
+	}
+
+	virtual ~FBasicTriangleSetSceneProxy() override
+	{
+		VertexBuffers.PositionVertexBuffer.ReleaseResource();
+		VertexBuffers.StaticMeshVertexBuffer.ReleaseResource();
+		VertexBuffers.ColorVertexBuffer.ReleaseResource();
+		IndexBuffer.ReleaseResource();
+		VertexFactory.ReleaseResource();
+	}
+
+	virtual void GetDynamicMeshElements(const TArray<const FSceneView*>& Views, const FSceneViewFamily& ViewFamily, uint32 VisibilityMap, FMeshElementCollector& Collector) const override
+	{
+		QUICK_SCOPE_CYCLE_COUNTER(STAT_OverlaySceneProxy_GetDynamicMeshElements);
+
+		for (int32 ViewIndex = 0; ViewIndex < Views.Num(); ViewIndex++)
+		{
+			if (VisibilityMap & (1 << ViewIndex))
+			{
+				for (const FBasicTriangleSetMeshBatchData& MeshBatchData : MeshBatchDatas)
+				{
+					FMeshBatch& Mesh = Collector.AllocateMesh();
+					FMeshBatchElement& BatchElement = Mesh.Elements[0];
+					BatchElement.IndexBuffer = &IndexBuffer;
+					Mesh.bWireframe = false;
+					Mesh.VertexFactory = &VertexFactory;
+					Mesh.MaterialRenderProxy = MeshBatchData.MaterialProxy;
+
+					FDynamicPrimitiveUniformBuffer& DynamicPrimitiveUniformBuffer = Collector.AllocateOneFrameResource<FDynamicPrimitiveUniformBuffer>();
+					DynamicPrimitiveUniformBuffer.Set(GetLocalToWorld(), GetLocalToWorld(), GetBounds(), GetLocalBounds(), true, false, AlwaysHasVelocity());
+					BatchElement.PrimitiveUniformBufferResource = &DynamicPrimitiveUniformBuffer.UniformBuffer;
+
+					BatchElement.FirstIndex = MeshBatchData.StartIndex;
+					BatchElement.NumPrimitives = MeshBatchData.NumPrimitives;
+					BatchElement.MinVertexIndex = MeshBatchData.MinVertexIndex;
+					BatchElement.MaxVertexIndex = MeshBatchData.MaxVertexIndex;
+					Mesh.ReverseCulling = IsLocalToWorldDeterminantNegative();
+					Mesh.Type = PT_TriangleList;
+					Mesh.DepthPriorityGroup = SDPG_World;
+					Mesh.bCanApplyViewModeOverrides = false;
+					Collector.AddMesh(ViewIndex, Mesh);
+				}
+			}
+		}
+	}
+
+	virtual FPrimitiveViewRelevance GetViewRelevance(const FSceneView* View) const override
+	{
+		FPrimitiveViewRelevance Result;
+		Result.bDrawRelevance = IsShown(View);
+		Result.bShadowRelevance = IsShadowCast(View);
+		Result.bDynamicRelevance = true;
+		Result.bRenderInMainPass = ShouldRenderInMainPass();
+		Result.bUsesLightingChannels = GetLightingChannelMask() != GetDefaultLightingChannelMask();
+		Result.bRenderCustomDepth = ShouldRenderCustomDepth();
+		Result.bTranslucentSelfShadow = bCastVolumetricTranslucentShadow;
+		MaterialRelevance.SetPrimitiveViewRelevance(Result);
+		Result.bVelocityRelevance = DrawsVelocity() && Result.bOpaque && Result.bRenderInMainPass;
+		return Result;
+	}
+
+	virtual bool CanBeOccluded() const override
+	{
+		return !MaterialRelevance.bDisableDepthTest;
+	}
+
+	virtual uint32 GetMemoryFootprint() const override { return sizeof(*this) + GetAllocatedSize(); }
+
+	virtual SIZE_T GetTypeHash() const override
+	{
+		static SIZE_T UniquePointer;
+		return reinterpret_cast<SIZE_T>(&UniquePointer);
+	}
+
+private:
+
+	// Triangle type specific implementation to get an FVector3f out of the generic point type.
+	static FVector3f ToFVector3f(const FVector2f& Vector) { return FVector3f(Vector.X, Vector.Y, 0.0f); }
+	static FVector3f ToFVector3f(const FVector3f& Vector) { return Vector; }
+
+	TArray<FBasicTriangleSetMeshBatchData> MeshBatchDatas;
+	FMaterialRelevance MaterialRelevance;
+	FLocalVertexFactory VertexFactory;
+	FStaticMeshVertexBuffers VertexBuffers;
+	FDynamicMeshIndexBuffer32 IndexBuffer;
+};
+
+
+UBasicTriangleSetComponentBase::UBasicTriangleSetComponentBase()
+{
+	CastShadow = false;
+	bSelectable = false;
+	PrimaryComponentTick.bCanEverTick = false;
+
+	UPrimitiveComponent::SetCollisionProfileName(UCollisionProfile::NoCollision_ProfileName);
+}
+
+void UBasicTriangleSetComponentBase::ClearComponent()
+{
+	MarkRenderStateDirty();
+	bBoundsDirty = true;
+}
+
+
+void UBasicTriangleSetComponentBase::SetTriangleMaterial(UMaterialInterface* InTriangleMaterial)
+{
+	TriangleMaterial = InTriangleMaterial;
+	UMeshComponent::SetMaterial(0, InTriangleMaterial);
+}
+
+void UBasicTriangleSetComponentBase::SetTriangleSetParameters(FColor InColor, const FVector3f& InNormal)
+{
+	Color = InColor;
+	Normal = InNormal;
+	bBoundsDirty = true;
+}
+
+
+FPrimitiveSceneProxy* UBasic2DTriangleSetComponent::CreateSceneProxy()
+{
+	return Elements.Num() > 0 ? new FBasicTriangleSetSceneProxy(this) : nullptr;
+}
+
+FPrimitiveSceneProxy* UBasic3DTriangleSetComponent::CreateSceneProxy()
+{
+	return Elements.Num() > 0 ? new FBasicTriangleSetSceneProxy(this) : nullptr;
+}
diff --git a/Engine/Plugins/Experimental/UVEditor/Source/UVEditorTools/Private/Selection/UVToolSelectionHighlightMechanic.cpp b/Engine/Plugins/Experimental/UVEditor/Source/UVEditorTools/Private/Selection/UVToolSelectionHighlightMechanic.cpp
--- a/Engine/Plugins/Experimental/UVEditor/Source/UVEditorTools/Private/Selection/UVToolSelectionHighlightMechanic.cpp
+++ b/Engine/Plugins/Experimental/UVEditor/Source/UVEditorTools/Private/Selection/UVToolSelectionHighlightMechanic.cpp
@@ -4,7 +4,9 @@
 
 #include "Actions/UVSeamSewAction.h"
 #include "Drawing/LineSetComponent.h"
-#include "Drawing/PointSetComponent.h"
+#include "Drawing/BasicPointSetComponent.h"
+#include "Drawing/BasicLineSetComponent.h"
+#include "Drawing/BasicTriangleSetComponent.h"
 #include "Drawing/PreviewGeometryActor.h"
 #include "Drawing/TriangleSetComponent.h"
 #include "DynamicMesh/DynamicMesh3.h"
@@ -39,7 +41,7 @@ void UUVToolSelectionHighlightMechanic::Initialize(UWorld* UnwrapWorld, UWorld*
 	UnwrapGeometryActor = UnwrapWorld->SpawnActor<APreviewGeometryActor>(
 		FVector::ZeroVector, FRotator(0, 0, 0), FActorSpawnParameters());
 
-	UnwrapTriangleSet = NewObject<UTriangleSetComponent>(UnwrapGeometryActor);
+	UnwrapTriangleSet = NewObject<UBasic2DTriangleSetComponent>(UnwrapGeometryActor);
 	// We are setting the TranslucencySortPriority here to handle the UV editor's use case in 2D
 	// where multiple translucent layers are drawn on top of each other but still need depth sorting.
 	UnwrapTriangleSet->TranslucencySortPriority = FUVEditorUXSettings::SelectionTriangleDepthBias;
@@ -47,38 +49,62 @@ void UUVToolSelectionHighlightMechanic::Initialize(UWorld* UnwrapWorld, UWorld*
 		FUVEditorUXSettings::SelectionTriangleFillColor,
 		FUVEditorUXSettings::SelectionTriangleDepthBias,
 		FUVEditorUXSettings::SelectionTriangleOpacity);
+	UnwrapTriangleSet->SetTriangleMaterial(TriangleSetMaterial);
+	UnwrapTriangleSet->SetTriangleSetParameters(FUVEditorUXSettings::SelectionTriangleFillColor, FVector3f(0, 0, 1));
 	UnwrapGeometryActor->SetRootComponent(UnwrapTriangleSet.Get());
 	UnwrapTriangleSet->RegisterComponent();
 
-	UnwrapLineSet = NewObject<ULineSetComponent>(UnwrapGeometryActor);
+	UnwrapLineSet = NewObject<UBasic2DLineSetComponent>(UnwrapGeometryActor);
 	UnwrapLineSet->SetLineMaterial(ToolSetupUtil::GetDefaultLineComponentMaterial(GetParentTool()->GetToolManager(), true));
+	UnwrapLineSet->SetLineSetParameters(FUVEditorUXSettings::SelectionTriangleWireframeColor,
+							     			 FUVEditorUXSettings::SelectionLineThickness,
+										     FUVEditorUXSettings::SelectionWireframeDepthBias);
 	UnwrapLineSet->AttachToComponent(UnwrapTriangleSet.Get(), FAttachmentTransformRules::KeepWorldTransform);
 	UnwrapLineSet->RegisterComponent();
 
-	UnwrapPairedEdgeLineSet = NewObject<ULineSetComponent>(UnwrapGeometryActor);
+	UnwrapPairedEdgeLineSet = NewObject<UBasic2DLineSetComponent>(UnwrapGeometryActor);
 	UnwrapPairedEdgeLineSet->SetLineMaterial(ToolSetupUtil::GetDefaultLineComponentMaterial(GetParentTool()->GetToolManager(), true));
 	UnwrapPairedEdgeLineSet->AttachToComponent(UnwrapTriangleSet.Get(), FAttachmentTransformRules::KeepWorldTransform);
 	UnwrapPairedEdgeLineSet->RegisterComponent();
 
-	SewEdgePairingLineSet = NewObject<ULineSetComponent>(UnwrapGeometryActor);
-	SewEdgePairingLineSet->SetLineMaterial(ToolSetupUtil::GetDefaultLineComponentMaterial(
+	SewEdgePairingLeftLineSet = NewObject<UBasic2DLineSetComponent>(UnwrapGeometryActor);
+	SewEdgePairingLeftLineSet->SetLineMaterial(ToolSetupUtil::GetDefaultLineComponentMaterial(
 		GetParentTool()->GetToolManager(), /*bDepthTested*/ true));
-	SewEdgePairingLineSet->AttachToComponent(UnwrapTriangleSet.Get(), FAttachmentTransformRules::KeepWorldTransform);
-	SewEdgePairingLineSet->RegisterComponent();
-	SewEdgePairingLineSet->SetVisibility(bPairedEdgeHighlightsEnabled);
+	SewEdgePairingLeftLineSet->SetLineSetParameters(FUVEditorUXSettings::SewSideLeftColor,
+		                                                 FUVEditorUXSettings::SewLineHighlightThickness,
+		                                                 FUVEditorUXSettings::SewLineDepthOffset);
+	SewEdgePairingLeftLineSet->AttachToComponent(UnwrapTriangleSet.Get(), FAttachmentTransformRules::KeepWorldTransform);
+	SewEdgePairingLeftLineSet->RegisterComponent();
+	SewEdgePairingLeftLineSet->SetVisibility(bPairedEdgeHighlightsEnabled);
+
+	SewEdgePairingRightLineSet = NewObject<UBasic2DLineSetComponent>(UnwrapGeometryActor);
+	SewEdgePairingRightLineSet->SetLineMaterial(ToolSetupUtil::GetDefaultLineComponentMaterial(
+		GetParentTool()->GetToolManager(), /*bDepthTested*/ true));
+	SewEdgePairingRightLineSet->SetLineSetParameters(FUVEditorUXSettings::SewSideRightColor,
+		                                                  FUVEditorUXSettings::SewLineHighlightThickness,
+		                                                  FUVEditorUXSettings::SewLineDepthOffset);
+	SewEdgePairingRightLineSet->AttachToComponent(UnwrapTriangleSet.Get(), FAttachmentTransformRules::KeepWorldTransform);
+	SewEdgePairingRightLineSet->RegisterComponent();
+	SewEdgePairingRightLineSet->SetVisibility(bPairedEdgeHighlightsEnabled);
 
 	// The unselected paired edges get their own, stationary, actor.
 	UnwrapStationaryGeometryActor = UnwrapWorld->SpawnActor<APreviewGeometryActor>(
 		FVector::ZeroVector, FRotator(0, 0, 0), FActorSpawnParameters());
-	SewEdgeUnselectedPairingLineSet = NewObject<ULineSetComponent>(UnwrapStationaryGeometryActor);
+	SewEdgeUnselectedPairingLineSet = NewObject<UBasic2DLineSetComponent>(UnwrapStationaryGeometryActor);
 	SewEdgeUnselectedPairingLineSet->SetLineMaterial(ToolSetupUtil::GetDefaultLineComponentMaterial(
 		GetParentTool()->GetToolManager(), /*bDepthTested*/ true));
+	SewEdgeUnselectedPairingLineSet->SetLineSetParameters(FUVEditorUXSettings::SewSideRightColor,
+		                                                       FUVEditorUXSettings::SewLineHighlightThickness,
+		                                                       FUVEditorUXSettings::SewLineDepthOffset);
 	UnwrapStationaryGeometryActor->SetRootComponent(SewEdgeUnselectedPairingLineSet.Get());
 	SewEdgeUnselectedPairingLineSet->RegisterComponent();
 	SewEdgeUnselectedPairingLineSet->SetVisibility(bPairedEdgeHighlightsEnabled);
 
-	UnwrapPointSet = NewObject<UPointSetComponent>(UnwrapGeometryActor);
+	UnwrapPointSet = NewObject<UBasic2DPointSetComponent>(UnwrapGeometryActor);
 	UnwrapPointSet->SetPointMaterial(ToolSetupUtil::GetDefaultPointComponentMaterial(GetParentTool()->GetToolManager(), true));
+	UnwrapPointSet->SetPointSetParameters(FUVEditorUXSettings::SelectionTriangleWireframeColor,
+											   FUVEditorUXSettings::SelectionPointThickness,
+											   FUVEditorUXSettings::SelectionWireframeDepthBias);
 	UnwrapPointSet->AttachToComponent(UnwrapTriangleSet.Get(), FAttachmentTransformRules::KeepWorldTransform);
 	UnwrapPointSet->RegisterComponent();
 
@@ -86,15 +112,21 @@ void UUVToolSelectionHighlightMechanic::Initialize(UWorld* UnwrapWorld, UWorld*
 	LivePreviewGeometryActor = LivePreviewWorld->SpawnActor<APreviewGeometryActor>(
 		FVector::ZeroVector, FRotator(0, 0, 0), FActorSpawnParameters());
 
-	LivePreviewLineSet = NewObject<ULineSetComponent>(LivePreviewGeometryActor);
+	LivePreviewLineSet = NewObject<UBasic3DLineSetComponent>(LivePreviewGeometryActor);
 	LivePreviewLineSet->SetLineMaterial(ToolSetupUtil::GetDefaultLineComponentMaterial(
 		GetParentTool()->GetToolManager(), /*bDepthTested*/ true));
+	LivePreviewLineSet->SetLineSetParameters(FUVEditorUXSettings::SelectionTriangleWireframeColor,
+	                                               FUVEditorUXSettings::LivePreviewHighlightThickness,
+		                                           FUVEditorUXSettings::LivePreviewHighlightDepthOffset);
 	LivePreviewGeometryActor->SetRootComponent(LivePreviewLineSet.Get());
 	LivePreviewLineSet->RegisterComponent();
 
-	LivePreviewPointSet = NewObject<UPointSetComponent>(LivePreviewGeometryActor);
+	LivePreviewPointSet = NewObject<UBasic3DPointSetComponent>(LivePreviewGeometryActor);
 	LivePreviewPointSet->SetPointMaterial(ToolSetupUtil::GetDefaultPointComponentMaterial(
 		GetParentTool()->GetToolManager(), /*bDepthTested*/ true));
+	LivePreviewPointSet->SetPointSetParameters(FUVEditorUXSettings::SelectionTriangleWireframeColor,
+													FUVEditorUXSettings::LivePreviewHighlightPointSize,
+													FUVEditorUXSettings::LivePreviewHighlightDepthOffset);
 	LivePreviewPointSet->AttachToComponent(LivePreviewLineSet.Get(), FAttachmentTransformRules::KeepWorldTransform);
 	LivePreviewPointSet->RegisterComponent();
 }
@@ -148,7 +180,8 @@ void UUVToolSelectionHighlightMechanic::RebuildUnwrapHighlight(
 	UnwrapTriangleSet->Clear();
 	UnwrapLineSet->Clear();
 	UnwrapPointSet->Clear();
-	SewEdgePairingLineSet->Clear();
+	SewEdgePairingRightLineSet->Clear();
+	SewEdgePairingLeftLineSet->Clear();
 	SewEdgeUnselectedPairingLineSet->Clear();
 	StaticPairedEdgeVidsPerMesh.Reset();
 
@@ -168,31 +201,36 @@ void UUVToolSelectionHighlightMechanic::RebuildUnwrapHighlight(
 		{
 			TRACE_CPUPROFILER_EVENT_SCOPE(UUVToolSelectionHighlightMechanic::AppendUnwrapHighlight_Triangle);
 
-			UnwrapTriangleSet->ReserveTriangles(Selection.SelectedIDs.Num());
-			UnwrapLineSet->ReserveLines(Selection.SelectedIDs.Num() * 3);
+			UBasic2DTriangleSetComponent* UnwrapTriangleSetPtr = UnwrapTriangleSet.Get();
+			UBasic2DLineSetComponent* UnwrapLineSetPtr = UnwrapLineSet.Get();
+			UnwrapTriangleSetPtr->ReserveElements(Selection.SelectedIDs.Num());
+			UnwrapLineSetPtr->ReserveElements(Selection.SelectedIDs.Num() * 3);
 			for (int32 Tid : Selection.SelectedIDs)
 			{
 				if (!ensure(Mesh.IsTriangle(Tid)))
 				{
 					continue;
 				}
-
-				FIndex3i Vids = Mesh.GetTriangle(Tid);
+				
 				FVector Points[3];
+				Mesh.GetTriVertices(Tid, Points[0], Points[1], Points[2]);
 				for (int i = 0; i < 3; ++i)
 				{
-					Points[i] = StartTransform.InverseTransformPosition(Mesh.GetVertex(Vids[i]));
+					Points[i] = StartTransform.InverseTransformPosition(Points[i]);
 				}
-				UnwrapTriangleSet->AddTriangle(Points[0], Points[1], Points[2], FVector(0, 0, 1), FUVEditorUXSettings::SelectionTriangleFillColor, TriangleSetMaterial);
+				UnwrapTriangleSetPtr->AddElement(FVector2f(Points[0].X, Points[0].Y),
+											  FVector2f(Points[1].X, Points[1].Y),
+											  FVector2f(Points[2].X, Points[2].Y));
+					
 				for (int i = 0; i < 3; ++i)
 				{
 					int NextIndex = (i + 1) % 3;
-					UnwrapLineSet->AddLine(Points[i], Points[NextIndex],
-						FUVEditorUXSettings::SelectionTriangleWireframeColor,
-						FUVEditorUXSettings::SelectionLineThickness,
-						FUVEditorUXSettings::SelectionWireframeDepthBias);
+					UnwrapLineSetPtr->AddElement(FVector2f(Points[i].X, Points[i].Y),
+						                      FVector2f(Points[NextIndex].X, Points[NextIndex].Y));
 				}
 			}
+			UnwrapTriangleSetPtr->MarkRenderStateDirty();
+			UnwrapLineSetPtr->MarkRenderStateDirty();
 		}
 		else if (Selection.Type == FUVToolSelection::EType::Edge)
 		{
@@ -204,21 +242,28 @@ void UUVToolSelectionHighlightMechanic::RebuildUnwrapHighlight(
 			const FDynamicMesh3& AppliedMesh = bUsePreviews ? *Selection.Target->AppliedPreview->PreviewMesh->GetMesh()
 				: *Selection.Target->AppliedCanonical;
 
-			UnwrapLineSet->ReserveLines(Selection.SelectedIDs.Num());
+			UBasic2DLineSetComponent* UnwrapLineSetPtr = UnwrapLineSet.Get();
+			UBasic2DLineSetComponent* SewEdgePairingLeftLineSetPtr = SewEdgePairingLeftLineSet.Get();
+			UBasic2DLineSetComponent* SewEdgePairingRightLineSetPtr = SewEdgePairingRightLineSet.Get();
+			UBasic2DLineSetComponent* SewEdgeUnselectedPairingLineSetPtr = SewEdgeUnselectedPairingLineSet.Get();
+
+			UnwrapLineSetPtr->ReserveElements(Selection.SelectedIDs.Num());
+			SewEdgePairingLeftLineSetPtr->ReserveElements(Selection.SelectedIDs.Num());
+			SewEdgePairingRightLineSetPtr->ReserveElements(Selection.SelectedIDs.Num());
+			SewEdgeUnselectedPairingLineSetPtr->ReserveElements(Selection.SelectedIDs.Num());
 			for (int32 Eid : Selection.SelectedIDs)
 			{
 				if (!ensure(Mesh.IsEdge(Eid)))
 				{
 					continue;
 				}
-
-				FIndex2i EdgeVids = Mesh.GetEdgeV(Eid);
-				UnwrapLineSet->AddLine(
-					StartTransform.InverseTransformPosition(Mesh.GetVertex(EdgeVids.A)),
-					StartTransform.InverseTransformPosition(Mesh.GetVertex(EdgeVids.B)),
-					FUVEditorUXSettings::SelectionTriangleWireframeColor,
-					FUVEditorUXSettings::SelectionLineThickness,
-					FUVEditorUXSettings::SelectionWireframeDepthBias);
+				
+				FVector Points[2];
+				Mesh.GetEdgeV(Eid, Points[0], Points[1]);
+				Points[0] = StartTransform.InverseTransformPosition(Points[0]);
+				Points[1] = StartTransform.InverseTransformPosition(Points[1]);
+				UnwrapLineSetPtr->AddElement(FVector2f(Points[0].X, Points[0].Y),
+					                      FVector2f(Points[1].X, Points[1].Y));
 
 				if (bPairedEdgeHighlightsEnabled)
 				{
@@ -236,53 +281,54 @@ void UUVToolSelectionHighlightMechanic::RebuildUnwrapHighlight(
 						continue;
 					}
 
-					FIndex2i Vids1 = Mesh.GetEdgeV(Eid);
-					SewEdgePairingLineSet->AddLine(
-						StartTransform.InverseTransformPosition(Mesh.GetVertex(Vids1.A)),
-						StartTransform.InverseTransformPosition(Mesh.GetVertex(Vids1.B)),
-						FUVEditorUXSettings::SewSideLeftColor, FUVEditorUXSettings::SewLineHighlightThickness, FUVEditorUXSettings::SewLineDepthOffset);
+					Mesh.GetEdgeV(Eid, Points[0], Points[1]);
+					Points[0] = StartTransform.InverseTransformPosition(Points[0]);
+					Points[1] = StartTransform.InverseTransformPosition(Points[1]);
+					SewEdgePairingLeftLineSetPtr->AddElement(FVector2f(Points[0].X, Points[0].Y),
+					                                      FVector2f(Points[1].X, Points[1].Y));
 
 					// The paired edge may need to go into a separate line set if it is not selected so that it does
 					// not get affected by transformations of the selected highlights in SetUnwrapHighlightTransform
 					FIndex2i Vids2 = Mesh.GetEdgeV(PairedEid);
+					Mesh.GetEdgeV(PairedEid, Points[0], Points[1]);
 					if (bPairedEdgeIsSelected)
 					{
-						SewEdgePairingLineSet->AddLine(
-							StartTransform.InverseTransformPosition(Mesh.GetVertex(Vids2.A)),
-							StartTransform.InverseTransformPosition(Mesh.GetVertex(Vids2.B)),
-							FUVEditorUXSettings::SewSideRightColor, FUVEditorUXSettings::SewLineHighlightThickness, 
-							FUVEditorUXSettings::SewLineDepthOffset);
+						Points[0] = StartTransform.InverseTransformPosition(Points[0]);
+						Points[1] = StartTransform.InverseTransformPosition(Points[1]);
+						SewEdgePairingRightLineSetPtr->AddElement(FVector2f(Points[0].X, Points[0].Y),
+							                                      FVector2f(Points[1].X, Points[1].Y));
 					}
 					else
 					{
 						StaticPairedEdgeVidsPerMesh.Last().Value.Add(TPair<int32, int32>(Vids2.A, Vids2.B));
-						SewEdgeUnselectedPairingLineSet->AddLine(
-							Mesh.GetVertex(Vids2.A),
-							Mesh.GetVertex(Vids2.B),
-							FUVEditorUXSettings::SewSideRightColor, FUVEditorUXSettings::SewLineHighlightThickness,
-							FUVEditorUXSettings::SewLineDepthOffset);
+						SewEdgeUnselectedPairingLineSetPtr->AddElement(FVector2f(Points[0].X, Points[0].Y),
+ 							                                           FVector2f(Points[1].X, Points[1].Y));						
 					}
 				}//end if visualizing paired edges
 			}//end for each edge
+			UnwrapLineSetPtr->MarkRenderStateDirty();
+			SewEdgePairingLeftLineSetPtr->MarkRenderStateDirty();
+			SewEdgePairingRightLineSetPtr->MarkRenderStateDirty();
+			SewEdgeUnselectedPairingLineSetPtr->MarkRenderStateDirty();
 		}
 		else if (Selection.Type == FUVToolSelection::EType::Vertex)
 		{
 			TRACE_CPUPROFILER_EVENT_SCOPE(MeshSelectionMechanic_RebuildDrawnElements_Vertex);
 
-			UnwrapPointSet->ReservePoints(Selection.SelectedIDs.Num());
-			for (int32 Vid : Selection.SelectedIDs)
+			UBasic2DPointSetComponent* UnwrapPointSetPtr = UnwrapPointSet.Get();
+			UnwrapPointSetPtr->ReserveElements(Selection.SelectedIDs.Num());
+
+			for (const int32 Vid : Selection.SelectedIDs)
 			{
 				if (!ensure(Mesh.IsVertex(Vid)))
 				{
 					continue;
 				}
 
-				FRenderablePoint PointToRender(StartTransform.InverseTransformPosition(Mesh.GetVertex(Vid)),
-					FUVEditorUXSettings::SelectionTriangleWireframeColor,
-					FUVEditorUXSettings::SelectionPointThickness,
-					FUVEditorUXSettings::SelectionWireframeDepthBias);
-				UnwrapPointSet->AddPoint(PointToRender);
+				const FVector3d Position = StartTransform.InverseTransformPosition(Mesh.GetVertex(Vid));
+				UnwrapPointSetPtr->AddElement(FVector2f(Position.X, Position.Y));
 			}
+			UnwrapPointSetPtr->MarkRenderStateDirty();
 		}
 	}
 }
@@ -295,8 +341,18 @@ void UUVToolSelectionHighlightMechanic::SetUnwrapHighlightTransform(const FTrans
 		UnwrapGeometryActor->SetActorTransform(Transform);
 	}
 	if (bPairedEdgeHighlightsEnabled && bRebuildStaticPairedEdges)
-	{
-		SewEdgeUnselectedPairingLineSet->Clear();
+	{		
+		int32 PairedEdgeCount = 0;
+		for (const TPair<TWeakObjectPtr<UUVEditorToolMeshInput>,
+			TArray<TPair<int32, int32>>>& MeshVidPairs : StaticPairedEdgeVidsPerMesh)
+		{
+			PairedEdgeCount += MeshVidPairs.Value.Num();
+		}
+
+		UBasic2DLineSetComponent* SewEdgeUnselectedPairingLineSetPtr = SewEdgeUnselectedPairingLineSet.Get();
+		SewEdgeUnselectedPairingLineSetPtr->Clear();
+		SewEdgeUnselectedPairingLineSetPtr->ReserveElements(PairedEdgeCount);
+
 		for (const TPair<TWeakObjectPtr<UUVEditorToolMeshInput>, 
 			TArray<TPair<int32, int32>>>& MeshVidPairs : StaticPairedEdgeVidsPerMesh)
 		{
@@ -315,12 +371,12 @@ void UUVToolSelectionHighlightMechanic::SetUnwrapHighlightTransform(const FTrans
 				{
 					continue;
 				}
-				SewEdgeUnselectedPairingLineSet->AddLine(
-					Mesh.GetVertex(VidPair.Key),
-					Mesh.GetVertex(VidPair.Value),
-					FUVEditorUXSettings::SewSideRightColor, FUVEditorUXSettings::SewLineHighlightThickness,
-					FUVEditorUXSettings::SewLineDepthOffset);
+				FVector3d VertA = Mesh.GetVertex(VidPair.Key);
+				FVector3d VertB = Mesh.GetVertex(VidPair.Value);
+				SewEdgeUnselectedPairingLineSetPtr->AddElement(FVector2f(VertA.X, VertA.Y),
+                                         					FVector2f(VertB.X, VertB.Y));
 			}
+			SewEdgeUnselectedPairingLineSetPtr->MarkRenderStateDirty();
 		}
 	}
 }
@@ -342,8 +398,11 @@ void UUVToolSelectionHighlightMechanic::RebuildAppliedHighlightFromUnwrapSelecti
 		return;
 	}
 
-	LivePreviewLineSet->Clear();
-	LivePreviewPointSet->Clear();
+	UBasic3DLineSetComponent* LivePreviewLineSetPtr = LivePreviewLineSet.Get();
+	UBasic3DPointSetComponent* LivePreviewPointSetPtr = LivePreviewPointSet.Get();
+
+	LivePreviewLineSetPtr->Clear();
+	LivePreviewPointSetPtr->Clear();
 
 	for (const FUVToolSelection& Selection : UnwrapSelections)
 	{
@@ -361,23 +420,22 @@ void UUVToolSelectionHighlightMechanic::RebuildAppliedHighlightFromUnwrapSelecti
 
 		FTransform MeshTransform = Target->AppliedPreview->PreviewMesh->GetTransform();
 
-		auto AppendEdgeLine = [this, &AppliedMesh, &MeshTransform](int32 AppliedEid)
+		auto AppendEdgeLine = [this, &AppliedMesh, &MeshTransform, LivePreviewLineSetPtr](int32 AppliedEid)
 		{
-			FVector3d Vert1, Vert2;
-			AppliedMesh.GetEdgeV(AppliedEid, Vert1, Vert2);
-
-			LivePreviewLineSet->AddLine(
-				MeshTransform.TransformPosition(Vert1),
-				MeshTransform.TransformPosition(Vert2),
-				FUVEditorUXSettings::SelectionTriangleWireframeColor,
-				FUVEditorUXSettings::LivePreviewHighlightThickness,
-				FUVEditorUXSettings::LivePreviewHighlightDepthOffset);
+			FVector Points[2];
+			AppliedMesh.GetEdgeV(AppliedEid, Points[0], Points[1]);
+			Points[0] = MeshTransform.TransformPosition(Points[0]);
+			Points[1] = MeshTransform.TransformPosition(Points[1]);
+			LivePreviewLineSetPtr->AddElement(
+				FVector3f(Points[0].X, Points[0].Y, Points[0].Z),
+				FVector3f(Points[1].X, Points[1].Y, Points[1].Z));
 		};
 
 		if (Selection.Type == FUVToolSelection::EType::Triangle)
 		{
 			TRACE_CPUPROFILER_EVENT_SCOPE(Triangle);
 
+			LivePreviewLineSet->ReserveElements(Selection.SelectedIDs.Num()*3);
 			for (int32 Tid : Selection.SelectedIDs)
 			{
 				if (!ensure(AppliedMesh.IsTriangle(Tid)))
@@ -400,10 +458,12 @@ void UUVToolSelectionHighlightMechanic::RebuildAppliedHighlightFromUnwrapSelecti
 					}
 				}//end for tri edges
 			}//end for selection tids
+			LivePreviewLineSetPtr->MarkRenderStateDirty();
 		}
 		else if (Selection.Type == FUVToolSelection::EType::Edge)
 		{
 			TRACE_CPUPROFILER_EVENT_SCOPE(Edge);
+			LivePreviewLineSet->ReserveElements(Selection.SelectedIDs.Num());
 
 			for (int32 UnwrapEid : Selection.SelectedIDs)
 			{
@@ -421,21 +481,19 @@ void UUVToolSelectionHighlightMechanic::RebuildAppliedHighlightFromUnwrapSelecti
 
 				AppendEdgeLine(AppliedEid);
 			}
+			LivePreviewLineSetPtr->MarkRenderStateDirty();
 		}
 		else if (Selection.Type == FUVToolSelection::EType::Vertex)
 		{
 			TRACE_CPUPROFILER_EVENT_SCOPE(Vertex);
 
-			for (int32 UnwrapVid : Selection.SelectedIDs)
-			{
-				FVector3d Position = AppliedMesh.GetVertex(
-					Target->UnwrapVidToAppliedVid(UnwrapVid));
+			LivePreviewPointSetPtr->ReserveElements(Selection.SelectedIDs.Num());
 
-				LivePreviewPointSet->AddPoint(Position,
-					FUVEditorUXSettings::SelectionTriangleWireframeColor,
-					FUVEditorUXSettings::LivePreviewHighlightPointSize,
-					FUVEditorUXSettings::LivePreviewHighlightDepthOffset);
+			for (const int32 UnwrapVid : Selection.SelectedIDs)
+			{
+				LivePreviewPointSetPtr->AddElement(static_cast<FVector3f>(AppliedMesh.GetVertex(Target->UnwrapVidToAppliedVid(UnwrapVid))));
 			}
+			LivePreviewPointSetPtr->MarkRenderStateDirty();
 		}
 	}//end for selection
 }
@@ -447,6 +505,9 @@ void UUVToolSelectionHighlightMechanic::AppendAppliedHighlight(const TArray<FUVT
 		return;
 	}
 
+	UBasic3DLineSetComponent* LivePreviewLineSetPtr = LivePreviewLineSet.Get();
+	UBasic3DPointSetComponent* LivePreviewPointSetPtr = LivePreviewPointSet.Get();
+
 	for (const FUVToolSelection& Selection : AppliedSelections)
 	{
 		if (!ensure(Selection.Target.IsValid() && Selection.Target->IsValid()))
@@ -461,17 +522,15 @@ void UUVToolSelectionHighlightMechanic::AppendAppliedHighlight(const TArray<FUVT
 
 		FTransform MeshTransform = Target->AppliedPreview->PreviewMesh->GetTransform();
 
-		auto AppendEdgeLine = [this, &AppliedMesh, &MeshTransform](int32 AppliedEid)
+		auto AppendEdgeLine = [this, &AppliedMesh, &MeshTransform, LivePreviewLineSetPtr](int32 AppliedEid)
 		{
-			FVector3d Vert1, Vert2;
-			AppliedMesh.GetEdgeV(AppliedEid, Vert1, Vert2);
-
-			LivePreviewLineSet->AddLine(
-				MeshTransform.TransformPosition(Vert1),
-				MeshTransform.TransformPosition(Vert2),
-				FUVEditorUXSettings::SelectionTriangleWireframeColor,
-				FUVEditorUXSettings::LivePreviewHighlightThickness,
-				FUVEditorUXSettings::LivePreviewHighlightDepthOffset);
+			FVector Points[2];
+			AppliedMesh.GetEdgeV(AppliedEid, Points[0], Points[1]);
+			Points[0] = MeshTransform.TransformPosition(Points[0]);
+			Points[1] = MeshTransform.TransformPosition(Points[1]);
+			LivePreviewLineSetPtr->AddElement(
+				FVector3f(Points[0].X, Points[0].Y, Points[0].Z),
+				FVector3f(Points[1].X, Points[1].Y, Points[1].Z));
 		};
 
 		if (Selection.Type == FUVToolSelection::EType::Triangle)
@@ -500,6 +559,7 @@ void UUVToolSelectionHighlightMechanic::AppendAppliedHighlight(const TArray<FUVT
 					}
 				}//end for tri edges
 			}//end for selection tids
+			LivePreviewLineSetPtr->MarkRenderStateDirty();
 		}
 		else if (Selection.Type == FUVToolSelection::EType::Edge)
 		{
@@ -514,20 +574,19 @@ void UUVToolSelectionHighlightMechanic::AppendAppliedHighlight(const TArray<FUVT
 
 				AppendEdgeLine(Eid);
 			}
+			LivePreviewLineSetPtr->MarkRenderStateDirty();
 		}
 		else if (Selection.Type == FUVToolSelection::EType::Vertex)
 		{
 			TRACE_CPUPROFILER_EVENT_SCOPE(Vertex);
 
-			for (int32 Vid : Selection.SelectedIDs)
-			{
-				FVector3d Position = AppliedMesh.GetVertex(Vid);
+			LivePreviewPointSetPtr->ReserveElements(Selection.SelectedIDs.Num());
 
-				LivePreviewPointSet->AddPoint(Position,
-					FUVEditorUXSettings::SelectionTriangleWireframeColor,
-					FUVEditorUXSettings::LivePreviewHighlightPointSize,
-					FUVEditorUXSettings::LivePreviewHighlightDepthOffset);
+			for (const int32 Vid : Selection.SelectedIDs)
+			{
+				LivePreviewPointSetPtr->AddElement(static_cast<FVector3f>(AppliedMesh.GetVertex(Vid)));
 			}
+			LivePreviewPointSetPtr->MarkRenderStateDirty();
 		}
 	}//end for selection
 }
@@ -536,7 +595,8 @@ void UUVToolSelectionHighlightMechanic::AppendAppliedHighlight(const TArray<FUVT
 void UUVToolSelectionHighlightMechanic::SetEnablePairedEdgeHighlights(bool bEnable)
 {
 	bPairedEdgeHighlightsEnabled = bEnable;
-	SewEdgePairingLineSet->SetVisibility(bEnable);
+	SewEdgePairingLeftLineSet->SetVisibility(bEnable);
+	SewEdgePairingRightLineSet->SetVisibility(bEnable);
 	SewEdgeUnselectedPairingLineSet->SetVisibility(bEnable);
 }
 
diff --git a/Engine/Plugins/Experimental/UVEditor/Source/UVEditorTools/Public/Drawing/BasicElementSet.h b/Engine/Plugins/Experimental/UVEditor/Source/UVEditorTools/Public/Drawing/BasicElementSet.h
new file mode 100644
--- /dev/null
+++ b/Engine/Plugins/Experimental/UVEditor/Source/UVEditorTools/Public/Drawing/BasicElementSet.h
@@ -0,0 +1,118 @@
+// Copyright Epic Games, Inc. All Rights Reserved.
+
+#pragma once
+
+#include "CoreMinimal.h"
+#include "Components/MeshComponent.h"
+
+class FBasicPointSetSceneProxy;
+
+
+/**
+ * A generic data storage class template for 3D or 2D vector data representing geometric elements composed of 1, 2, or 3 points, nominally
+ * points, lines and triangles respectively. 
+ * 
+ * 
+ * To support both 2D and 3D with minimal code duplication, the heigher level point/line/triangle set components are composed from the UBasic[*]SetComponentBase class for the component
+ * specific functionality, and an instantiation of the TBasicElementSet template class for either FVector2f or FVector3f that encapsulates all data storage specific
+ * functionality. The final derived classes UBasic[2/3]D[Point/Line/Triangle]SetComponent contain identical boilerplate code to facilitate calls that
+ * need to go both to the component base class and the point base class, e.g. Clear() will both mark the component as dirty and delete all points.  
+ */
+
+/** Generic storage for sets of geometry. */
+template <typename VectorType, int32 ElementDataSize>
+class TBasicElementSet
+{
+public:
+
+	using GeometryType = VectorType;
+	static const int32 ElementSize = ElementDataSize;
+
+	VectorType& operator[](int32 Index)
+	{
+		return Elements[Index];
+	}
+
+	const VectorType& operator[](int32 Index) const
+	{
+		return Elements[Index];
+	}
+
+	/** Current number of points. */
+	int32 NumElements() const
+	{
+		return Num;
+	}
+
+	/** Reserve enough memory for up to the given maximum number of points. This function needs to be called before any calls to AddPoint(). */
+	void ReserveElements(const int32 MaxNum)
+	{
+		Elements.SetNumUninitialized(MaxNum * ElementDataSize);
+	}
+
+	/** Add a point to be rendered using the component. */
+	template<typename... Targs>
+	void AddElement(Targs... ElementPoints)
+	{
+		static_assert(sizeof...(Targs) == ElementDataSize);
+		checkSlow(Num < Elements.Num());
+		AddElementInternal(ElementPoints...);
+		Num++;
+		bElementsDirty = true;
+	}
+
+	void ClearElements()
+	{
+		Elements.Reset();
+		Num = 0;
+		bElementsDirty = true;
+	}
+
+protected:
+
+	template<typename T, typename... Targs>
+	void AddElementInternal(const T& ElementPoint, Targs... ElementPoints)
+	{
+		int32 Offset = (ElementDataSize - sizeof...(Targs));
+		Elements[Num * ElementDataSize + Offset - 1] = ElementPoint;
+		AddElementInternal(ElementPoints...);
+	}
+
+	template<typename T>
+	void AddElementInternal(const T& ElementPoint)
+	{
+		Elements[Num* ElementDataSize + ElementDataSize-1] = ElementPoint;
+	}
+
+	FBox CalcElementsBox() const
+	{
+		FBox Box(ForceInit);
+		for (int32 i = 0; i < Num; ++i)
+		{
+			for (int32 j = 0; j < ElementDataSize; ++j)
+			{
+				// Adding a point to an FBox is specific to the point type, and thus we call out to a helper function with dedicated template instantiations.
+				AddPointToBox(Elements[i*ElementDataSize + j], Box);
+			}
+		}
+		return Box;
+	}
+
+	TArray<VectorType> Elements;
+	int32 Num = 0;
+
+	bool bElementsDirty = true;
+
+private:
+
+	// Helper function to generically add a point to an FBox.
+	static void AddPointToBox(const VectorType& Point, FBox& Box);
+};
+
+// Helper function template instantiations for FVector2f and FVector3f
+template <> inline void TBasicElementSet<FVector2f, 1>::AddPointToBox(const FVector2f& Point, FBox& Box) { Box += FVector3d(Point.X, Point.Y, 0.0); }
+template <> inline void TBasicElementSet<FVector3f, 1>::AddPointToBox(const FVector3f& Point, FBox& Box) { Box += static_cast<FVector3d>(Point); }
+template <> inline void TBasicElementSet<FVector2f, 2>::AddPointToBox(const FVector2f& Point, FBox& Box) { Box += FVector3d(Point.X, Point.Y, 0.0); }
+template <> inline void TBasicElementSet<FVector3f, 2>::AddPointToBox(const FVector3f& Point, FBox& Box) { Box += static_cast<FVector3d>(Point); }
+template <> inline void TBasicElementSet<FVector2f, 3>::AddPointToBox(const FVector2f& Point, FBox& Box) { Box += FVector3d(Point.X, Point.Y, 0.0); }
+template <> inline void TBasicElementSet<FVector3f, 3>::AddPointToBox(const FVector3f& Point, FBox& Box) { Box += static_cast<FVector3d>(Point); }
diff --git a/Engine/Plugins/Experimental/UVEditor/Source/UVEditorTools/Public/Drawing/BasicLineSetComponent.h b/Engine/Plugins/Experimental/UVEditor/Source/UVEditorTools/Public/Drawing/BasicLineSetComponent.h
new file mode 100644
--- /dev/null
+++ b/Engine/Plugins/Experimental/UVEditor/Source/UVEditorTools/Public/Drawing/BasicLineSetComponent.h
@@ -0,0 +1,157 @@
+// Copyright Epic Games, Inc. All Rights Reserved.
+
+#pragma once
+
+#include "CoreMinimal.h"
+#include "Components/MeshComponent.h"
+#include "BasicElementSet.h"
+#include "BasicLineSetComponent.generated.h"
+
+class FBasicLineSetSceneProxy;
+
+
+/**
+ * UBasic2DLineSetComponent and UBasic3DLineSetComponent are components that draw a set of lines as small squares. The 2D version stores XY coordinates with
+ * Z set to 0 whereas the 3D version stores XYZ coordinates.
+ * All lines have the same color, size, and depth bias based on what is set via SetLineSetParameters(). The lines are drawn as two triangles, i.e. a
+ * square, orthogonal to the view direction. The actual line size is calculated in the shader using a custom material.
+ *
+ * Before inserting any lines via AddElement(), the maximum number of lines that will be added needs to be declared via ReserveElements().
+ * The actual number of lines that were added is available via NumElements().
+ *
+ * To support both 2D and 3D with minimal code duplication, the line set components are composed from the UBasicLineSetComponentBase class for the component
+ * specific functionality, and an instantiation of the TBasicElementSet template class for either FVector2f or FVector3f that encapsulates all line specific
+ * functionality. The final derived classes UBasic2DLineSetComponent and UBasic3DLineSetComponent contain identical boilerplate code to facilitate calls that
+ * need to go both to the component base class and the line base class, e.g. Clear() will both mark the component as dirty and delete all lines.
+ */
+
+/** Base class for component specific functionality independent of the type of line stored in the component. */
+UCLASS()
+class UVEDITORTOOLS_API UBasicLineSetComponentBase : public UMeshComponent
+{
+	GENERATED_BODY()
+
+public:
+
+	UBasicLineSetComponentBase();
+
+	/** Clears the component state and marks component as dirty. */
+	void ClearComponent();
+
+	/** Specify material that handles lines. */
+	void SetLineMaterial(UMaterialInterface* InLineMaterial);
+
+	/** Set per line material parameters that are uniform for all lines. */
+	void SetLineSetParameters(FColor InColor, float InSize, float InDepthBias);
+
+protected:
+
+	//~ UMeshComponent Interface.
+	virtual int32 GetNumMaterials() const override { return 1; }
+
+	/** Update component bounds using a provided function that determines the box around all line positions.
+	 *  This unifies how the bounds are calculated while keeping the calculation of the box around all line positions generic. */
+	template <typename CalcLineBoxFn>
+	FBoxSphereBounds UpdateComponentBounds(const FTransform& LocalToWorld, bool bLinesDirty, CalcLineBoxFn&& CalcLineBox) const
+	{
+		bBoundsDirty |= bLinesDirty;
+		if (bBoundsDirty)
+		{
+			// Calculating an FBox around a set of lines is specific to the line type, and thus we call out to a provided type specific lambda.
+			const FBox Box = CalcLineBox();
+
+			Bounds = FBoxSphereBounds(Box);
+			bBoundsDirty = false;
+
+			// TODO: This next bit is not ideal because the line size is specified in onscreen pixels,
+			// so the true amount by which we would need to expand bounds depends on camera location, FOV, etc.
+			// We mainly do this as a hack against a problem in orthographic viewports, which cull small items
+			// based on their bounds, and a set consisting of a single line will always be culled due
+			// to having 0-sized bounds. It's worth noting that when zooming out sufficiently far, the
+			// line will still be culled even with this hack, however.
+			// The proper solution is to be able to opt out of the orthographic culling behavior, which is something
+			// we need to add.
+			if (Box.IsValid)
+			{
+				Bounds = Bounds.ExpandBy(Size);
+			}
+		}
+		return Bounds.TransformBy(LocalToWorld);
+	}
+
+	UPROPERTY()
+		TObjectPtr<const UMaterialInterface> LineMaterial;
+
+	UPROPERTY()
+		mutable FBoxSphereBounds Bounds;
+
+	UPROPERTY()
+		mutable bool bBoundsDirty = true;
+
+	UPROPERTY()
+		FColor Color;
+
+	UPROPERTY()
+		float Size = 1.0f;
+
+	UPROPERTY()
+		float DepthBias = 0.0f;
+};
+
+/** Instantiation of a basic line set component in 2D using FVector2f for line positions. */
+UCLASS()
+class UVEDITORTOOLS_API UBasic2DLineSetComponent final : public UBasicLineSetComponentBase, public TBasicElementSet<FVector2f, 2>
+{
+	GENERATED_BODY()
+
+public:
+	/** Clear all lines and component state. */
+	void Clear()
+	{
+		ClearElements();
+		ClearComponent();
+	}
+
+private:
+
+	//~ UPrimitiveComponent Interface.
+	virtual FPrimitiveSceneProxy * CreateSceneProxy() override;
+
+	//~ USceneComponent Interface.
+	virtual FBoxSphereBounds CalcBounds(const FTransform& LocalToWorld) const override
+	{
+		// The call to CalcElementsBox() is specific to the element type but already available in the generic implementation in TBasicElementSet<>.
+		return UpdateComponentBounds(LocalToWorld, bElementsDirty, [this]() { return CalcElementsBox(); });
+	}
+
+	friend class FBasicLineSetSceneProxy;
+};
+
+/** Instantiation of a basic line set component in 3D using FVector3f for line positions. */
+UCLASS()
+class UVEDITORTOOLS_API UBasic3DLineSetComponent final : public UBasicLineSetComponentBase, public TBasicElementSet<FVector3f, 2>
+{
+	GENERATED_BODY()
+
+public:
+	/** Clear all lines and component state. */
+	void Clear()
+	{
+		ClearElements();
+		ClearComponent();
+	}
+
+private:
+
+	//~ UPrimitiveComponent Interface.
+	virtual FPrimitiveSceneProxy * CreateSceneProxy() override;
+
+	//~ USceneComponent Interface.
+	virtual FBoxSphereBounds CalcBounds(const FTransform& LocalToWorld) const override
+	{
+		// The call to CalcElementsBox() is specific to the element type but already available in the generic implementation in TBasicElementSet<>.
+		return UpdateComponentBounds(LocalToWorld, bElementsDirty, [this]() { return CalcElementsBox(); });
+	}
+
+	friend class FBasicLineSetSceneProxy;
+};
diff --git a/Engine/Plugins/Experimental/UVEditor/Source/UVEditorTools/Public/Drawing/BasicPointSetComponent.h b/Engine/Plugins/Experimental/UVEditor/Source/UVEditorTools/Public/Drawing/BasicPointSetComponent.h
new file mode 100644
--- /dev/null
+++ b/Engine/Plugins/Experimental/UVEditor/Source/UVEditorTools/Public/Drawing/BasicPointSetComponent.h
@@ -0,0 +1,157 @@
+// Copyright Epic Games, Inc. All Rights Reserved.
+
+#pragma once
+
+#include "CoreMinimal.h"
+#include "Components/MeshComponent.h"
+#include "BasicElementSet.h"
+#include "BasicPointSetComponent.generated.h"
+
+class FBasicPointSetSceneProxy;
+
+
+/**
+ * UBasic2DPointSetComponent and UBasic3DPointSetComponent are components that draw a set of points as small squares. The 2D version stores XY coordinates with
+ * Z set to 0 whereas the 3D version stores XYZ coordinates.
+ * All points have the same color, size, and depth bias based on what is set via SetPointSetParameters(). The points are drawn as two triangles, i.e. a
+ * square, orthogonal to the view direction. The actual point size is calculated in the shader using a custom material.
+ *
+ * Before inserting any points via AddElement(), the maximum number of points that will be added needs to be declared via ReserveElements().
+ * The actual number of points that were added is available via NumElements().
+ *
+ * To support both 2D and 3D with minimal code duplication, the point set components are composed from the UBasicPointSetComponentBase class for the component
+ * specific functionality, and an instantiation of the TBasicElementSet template class for either FVector2f or FVector3f that encapsulates all point specific
+ * functionality. The final derived classes UBasic2DPointSetComponent and UBasic3DPointSetComponent contain identical boilerplate code to facilitate calls that
+ * need to go both to the component base class and the point base class, e.g. Clear() will both mark the component as dirty and delete all points.  
+ */
+
+/** Base class for component specific functionality independent of the type of point stored in the component. */
+UCLASS()
+class UVEDITORTOOLS_API UBasicPointSetComponentBase : public UMeshComponent
+{
+	GENERATED_BODY()
+
+public:
+
+	UBasicPointSetComponentBase();
+
+	/** Clears the component state and marks component as dirty. */
+	void ClearComponent();
+
+	/** Specify material that handles points. */
+	void SetPointMaterial(UMaterialInterface* InPointMaterial);
+
+	/** Set per point material parameters that are uniform for all points. */
+	void SetPointSetParameters(FColor InColor, float InSize, float InDepthBias);
+
+protected:
+
+	//~ UMeshComponent Interface.
+	virtual int32 GetNumMaterials() const override { return 1; }
+
+	/** Update component bounds using a provided function that determines the box around all point positions.
+	 *  This unifies how the bounds are calculated while keeping the calculation of the box around all point positions generic. */
+	template <typename CalcPointBoxFn>
+	FBoxSphereBounds UpdateComponentBounds(const FTransform& LocalToWorld, bool bPointsDirty, CalcPointBoxFn&& CalcPointBox) const
+	{
+		bBoundsDirty |= bPointsDirty;
+		if (bBoundsDirty)
+		{
+			// Calculating an FBox around a set of points is specific to the point type, and thus we call out to a provided type specific lambda.
+			const FBox Box = CalcPointBox();
+
+			Bounds = FBoxSphereBounds(Box);
+			bBoundsDirty = false;
+
+			// TODO: This next bit is not ideal because the point size is specified in onscreen pixels,
+			// so the true amount by which we would need to expand bounds depends on camera location, FOV, etc.
+			// We mainly do this as a hack against a problem in orthographic viewports, which cull small items
+			// based on their bounds, and a set consisting of a single point will always be culled due
+			// to having 0-sized bounds. It's worth noting that when zooming out sufficiently far, the
+			// point will still be culled even with this hack, however.
+			// The proper solution is to be able to opt out of the orthographic culling behavior, which is something
+			// we need to add.
+			if (Box.IsValid)
+			{
+				Bounds = Bounds.ExpandBy(Size);
+			}
+		}
+		return Bounds.TransformBy(LocalToWorld);
+	}
+
+	UPROPERTY()
+	TObjectPtr<const UMaterialInterface> PointMaterial;
+
+	UPROPERTY()
+	mutable FBoxSphereBounds Bounds;
+
+	UPROPERTY()
+	mutable bool bBoundsDirty = true;
+
+	UPROPERTY()
+	FColor Color;
+
+	UPROPERTY()
+	float Size = 1.0f;
+
+	UPROPERTY()
+	float DepthBias = 0.0f;
+};
+
+/** Instantiation of a basic point set component in 2D using FVector2f for point positions. */
+UCLASS()
+class UVEDITORTOOLS_API UBasic2DPointSetComponent final : public UBasicPointSetComponentBase, public TBasicElementSet<FVector2f, 1>
+{
+	GENERATED_BODY()
+
+public:
+	/** Clear all points and component state. */
+	void Clear()
+	{
+		ClearElements();
+		ClearComponent();
+	}
+
+private:
+
+	//~ UPrimitiveComponent Interface.
+	virtual FPrimitiveSceneProxy* CreateSceneProxy() override;
+
+	//~ USceneComponent Interface.
+	virtual FBoxSphereBounds CalcBounds(const FTransform& LocalToWorld) const override
+	{
+		// The call to CalcElementsBox() is specific to the element type but already available in the generic implementation in TBasicElementSet<>.
+		return UpdateComponentBounds(LocalToWorld, bElementsDirty, [this]() { return CalcElementsBox(); });
+	}
+
+	friend class FBasicPointSetSceneProxy;
+};
+
+/** Instantiation of a basic point set component in 3D using FVector3f for point positions. */
+UCLASS()
+class UVEDITORTOOLS_API UBasic3DPointSetComponent final : public UBasicPointSetComponentBase, public TBasicElementSet<FVector3f, 1>
+{
+	GENERATED_BODY()
+
+public:
+	/** Clear all points and component state. */
+	void Clear()
+	{
+		ClearElements();
+		ClearComponent();
+	}
+
+private:
+
+	//~ UPrimitiveComponent Interface.
+	virtual FPrimitiveSceneProxy* CreateSceneProxy() override;
+
+	//~ USceneComponent Interface.
+	virtual FBoxSphereBounds CalcBounds(const FTransform& LocalToWorld) const override
+	{
+		// The call to CalcElementsBox() is specific to the element type but already available in the generic implementation in TBasicElementSet<>.
+		return UpdateComponentBounds(LocalToWorld, bElementsDirty, [this]() { return CalcElementsBox(); });
+	}
+
+	friend class FBasicPointSetSceneProxy;
+};
diff --git a/Engine/Plugins/Experimental/UVEditor/Source/UVEditorTools/Public/Drawing/BasicTriangleSetComponent.h b/Engine/Plugins/Experimental/UVEditor/Source/UVEditorTools/Public/Drawing/BasicTriangleSetComponent.h
new file mode 100644
--- /dev/null
+++ b/Engine/Plugins/Experimental/UVEditor/Source/UVEditorTools/Public/Drawing/BasicTriangleSetComponent.h
@@ -0,0 +1,141 @@
+// Copyright Epic Games, Inc. All Rights Reserved.
+
+#pragma once
+
+#include "CoreMinimal.h"
+#include "Components/MeshComponent.h"
+#include "BasicElementSet.h"
+#include "BasicTriangleSetComponent.generated.h"
+
+class FBasicTriangleSetSceneProxy;
+
+
+/**
+ * UBasic2DTriangleSetComponent and UBasic3DTriangleSetComponent are components that draw a set of lines as small squares. The 2D version stores XY coordinates with
+ * Z set to 0 whereas the 3D version stores XYZ coordinates.
+ * All lines have the same color, size, and depth bias based on what is set via SetTriangleSetParameters(). The lines are drawn as two triangles, i.e. a
+ * square, orthogonal to the view direction. The actual line size is calculated in the shader using a custom material.
+ *
+ * Before inserting any lines via AddElement(), the maximum number of lines that will be added needs to be declared via ReserveElements().
+ * The actual number of lines that were added is available via NumElements().
+ *
+ * To support both 2D and 3D with minimal code duplication, the line set components are composed from the UBasicTriangleSetComponentBase class for the component
+ * specific functionality, and an instantiation of the TBasicElementSet template class for either FVector2f or FVector3f that encapsulates all line specific
+ * functionality. The final derived classes UBasic2DTriangleSetComponent and UBasic3DTriangleSetComponent contain identical boilerplate code to facilitate calls that
+ * need to go both to the component base class and the line base class, e.g. Clear() will both mark the component as dirty and delete all lines.
+ */
+
+ /** Base class for component specific functionality independent of the type of line stored in the component. */
+UCLASS()
+class UVEDITORTOOLS_API UBasicTriangleSetComponentBase : public UMeshComponent
+{
+	GENERATED_BODY()
+
+public:
+
+	UBasicTriangleSetComponentBase();
+
+	/** Clears the component state and marks component as dirty. */
+	void ClearComponent();
+
+	/** Specify material that handles lines. */
+	void SetTriangleMaterial(UMaterialInterface* InTriangleMaterial);
+
+	/** Set per line material parameters that are uniform for all lines. */
+	void SetTriangleSetParameters(FColor InColor, const FVector3f& NormalIn);
+
+protected:
+
+	//~ UMeshComponent Interface.
+	virtual int32 GetNumMaterials() const override { return 1; }
+
+	/** Update component bounds using a provided function that determines the box around all line positions.
+	 *  This unifies how the bounds are calculated while keeping the calculation of the box around all line positions generic. */
+	template <typename CalcTriangleBoxFn>
+	FBoxSphereBounds UpdateComponentBounds(const FTransform& LocalToWorld, bool bTrianglesDirty, CalcTriangleBoxFn&& CalcTriangleBox) const
+	{
+		bBoundsDirty |= bTrianglesDirty;
+		if (bBoundsDirty)
+		{
+			// Calculating an FBox around a set of lines is specific to the line type, and thus we call out to a provided type specific lambda.
+			const FBox Box = CalcTriangleBox();
+
+			Bounds = FBoxSphereBounds(Box);
+			bBoundsDirty = false;
+		}
+		return Bounds.TransformBy(LocalToWorld);
+	}
+
+	UPROPERTY()
+	TObjectPtr<const UMaterialInterface> TriangleMaterial;
+
+	UPROPERTY()
+	mutable FBoxSphereBounds Bounds;
+
+	UPROPERTY()
+	mutable bool bBoundsDirty = true;
+
+	UPROPERTY()
+	FColor Color;
+
+	UPROPERTY()
+	FVector3f Normal;
+};
+
+/** Instantiation of a basic line set component in 2D using FVector2f for line positions. */
+UCLASS()
+class UVEDITORTOOLS_API UBasic2DTriangleSetComponent final : public UBasicTriangleSetComponentBase, public TBasicElementSet<FVector2f, 3>
+{
+	GENERATED_BODY()
+
+public:
+	/** Clear all lines and component state. */
+	void Clear()
+	{
+		ClearElements();
+		ClearComponent();
+	}
+
+private:
+
+	//~ UPrimitiveComponent Interface.
+	virtual FPrimitiveSceneProxy * CreateSceneProxy() override;
+
+	//~ USceneComponent Interface.
+	virtual FBoxSphereBounds CalcBounds(const FTransform& LocalToWorld) const override
+	{
+		// The call to CalcElementsBox() is specific to the element type but already available in the generic implementation in TBasicElementSet<>.
+		return UpdateComponentBounds(LocalToWorld, bElementsDirty, [this]() { return CalcElementsBox(); });
+	}
+
+	friend class FBasicTriangleSetSceneProxy;
+};
+
+/** Instantiation of a basic line set component in 3D using FVector3f for line positions. */
+UCLASS()
+class UVEDITORTOOLS_API UBasic3DTriangleSetComponent final : public UBasicTriangleSetComponentBase, public TBasicElementSet<FVector3f, 3>
+{
+	GENERATED_BODY()
+
+public:
+	/** Clear all lines and component state. */
+	void Clear()
+	{
+		ClearElements();
+		ClearComponent();
+	}
+
+private:
+
+	//~ UPrimitiveComponent Interface.
+	virtual FPrimitiveSceneProxy * CreateSceneProxy() override;
+
+	//~ USceneComponent Interface.
+	virtual FBoxSphereBounds CalcBounds(const FTransform& LocalToWorld) const override
+	{
+		// The call to CalcElementsBox() is specific to the element type but already available in the generic implementation in TBasicElementSet<>.
+		return UpdateComponentBounds(LocalToWorld, bElementsDirty, [this]() { return CalcElementsBox(); });
+	}
+
+	friend class FBasicTriangleSetSceneProxy;
+};
diff --git a/Engine/Plugins/Experimental/UVEditor/Source/UVEditorTools/Public/Selection/UVToolSelectionHighlightMechanic.h b/Engine/Plugins/Experimental/UVEditor/Source/UVEditorTools/Public/Selection/UVToolSelectionHighlightMechanic.h
--- a/Engine/Plugins/Experimental/UVEditor/Source/UVEditorTools/Public/Selection/UVToolSelectionHighlightMechanic.h
+++ b/Engine/Plugins/Experimental/UVEditor/Source/UVEditorTools/Public/Selection/UVToolSelectionHighlightMechanic.h
@@ -14,9 +14,11 @@ PREDECLARE_GEOMETRY(class FUVToolSelection);
 
 class APreviewGeometryActor;
 class UMaterialInstanceDynamic;
-class UTriangleSetComponent;
-class ULineSetComponent;
-class UPointSetComponent;
+class UBasic2DPointSetComponent;
+class UBasic3DPointSetComponent;
+class UBasic2DLineSetComponent;
+class UBasic3DLineSetComponent;
+class UBasic2DTriangleSetComponent;
 class UUVEditorToolMeshInput;
 class UWorld;
 
@@ -96,15 +98,16 @@ protected:
 	UPROPERTY()
 	TObjectPtr<UMaterialInstanceDynamic> TriangleSetMaterial = nullptr;
 
-	TWeakObjectPtr<UTriangleSetComponent> UnwrapTriangleSet = nullptr;
-	TWeakObjectPtr<ULineSetComponent> UnwrapLineSet = nullptr;
-	TWeakObjectPtr<ULineSetComponent> SewEdgePairingLineSet = nullptr;
-	TWeakObjectPtr<UPointSetComponent> UnwrapPointSet = nullptr;
-	TWeakObjectPtr<ULineSetComponent> UnwrapPairedEdgeLineSet = nullptr;
+	TWeakObjectPtr<UBasic2DTriangleSetComponent> UnwrapTriangleSet = nullptr;
+	TWeakObjectPtr<UBasic2DLineSetComponent> UnwrapLineSet = nullptr;
+	TWeakObjectPtr<UBasic2DLineSetComponent> SewEdgePairingLeftLineSet = nullptr;
+	TWeakObjectPtr<UBasic2DLineSetComponent> SewEdgePairingRightLineSet = nullptr;
+	TWeakObjectPtr<UBasic2DPointSetComponent> UnwrapPointSet = nullptr;
+	TWeakObjectPtr<UBasic2DLineSetComponent> UnwrapPairedEdgeLineSet = nullptr;
 
 	UPROPERTY()
 	TObjectPtr<APreviewGeometryActor> UnwrapStationaryGeometryActor = nullptr;
-	TWeakObjectPtr<ULineSetComponent> SewEdgeUnselectedPairingLineSet = nullptr;
+	TWeakObjectPtr<UBasic2DLineSetComponent> SewEdgeUnselectedPairingLineSet = nullptr;
 
 	// Per mesh, gives a list of endpoint vid pairs for the unselected edges that
 	// are part of edge pairs, for rebuilding when necessary.
@@ -113,8 +116,8 @@ protected:
 
 	UPROPERTY()
 	TObjectPtr<APreviewGeometryActor> LivePreviewGeometryActor = nullptr;
-	TWeakObjectPtr<ULineSetComponent> LivePreviewLineSet = nullptr;
-	TWeakObjectPtr<UPointSetComponent> LivePreviewPointSet = nullptr;
+	TWeakObjectPtr<UBasic3DLineSetComponent> LivePreviewLineSet = nullptr;
+	TWeakObjectPtr<UBasic3DPointSetComponent> LivePreviewPointSet = nullptr;
 
 	bool bPairedEdgeHighlightsEnabled = true;
 };
diff --git a/Engine/Plugins/Experimental/UVEditor/Source/UVEditorTools/UVEditorTools.Build.cs b/Engine/Plugins/Experimental/UVEditor/Source/UVEditorTools/UVEditorTools.Build.cs
--- a/Engine/Plugins/Experimental/UVEditor/Source/UVEditorTools/UVEditorTools.Build.cs
+++ b/Engine/Plugins/Experimental/UVEditor/Source/UVEditorTools/UVEditorTools.Build.cs
@@ -43,6 +43,8 @@ public class UVEditorTools : ModuleRules
 				"Engine",
 				"Slate",
 				"SlateCore",
+				"RenderCore",
+				"RHI",
 					
 				"MeshModelingToolsExp",
 				"DynamicMesh",
