UVEditor: Remove duplicate "UV Editor" entry from "Actor > Asset Tools" menu

diff --git a/Engine/Plugins/Experimental/UVEditor/Source/UVEditor/Private/UVEditorModule.cpp b/Engine/Plugins/Experimental/UVEditor/Source/UVEditor/Private/UVEditorModule.cpp
--- a/Engine/Plugins/Experimental/UVEditor/Source/UVEditor/Private/UVEditorModule.cpp
+++ b/Engine/Plugins/Experimental/UVEditor/Source/UVEditor/Private/UVEditorModule.cpp
@@ -175,19 +175,6 @@ void FUVEditorModule::RegisterMenus()
 		Section.AddMenuEntryWithCommandList(FUVEditorCommands::Get().OpenUVEditor, CommandListToBind);
 	};
 
-	// Add UV Editor to "Actor" > "Asset Tools" sub menu.
-	{
-		UToolMenu* AssetToolsSubMenu = UToolMenus::Get()->ExtendMenu("LevelEditor.ActorContextMenu.AssetToolsSubMenu");
-		FToolMenuSection& AssetToolsSection = AssetToolsSubMenu->AddSection("UVEditorCommands", TAttribute<FText>(), FToolMenuInsert("AssetTools",
-			                                                                    EToolMenuInsertType::After));
-		AssetToolsSection.AddDynamicEntry("OpenUVEditor", FNewToolMenuSectionDelegate::CreateLambda(
-			                                  [&SetupMenuEntryParameters, &AddMenuEntry](FToolMenuSection& Section)
-			                                  {
-				                                  const FMenuEntryParameters MenuEntryParameters = SetupMenuEntryParameters();
-				                                  AddMenuEntry(Section, MenuEntryParameters);
-			                                  }));
-	}
-
 	// Add UV Editor to actor context menu.
 	{
 		UToolMenu* ActorContextMenu = UToolMenus::Get()->ExtendMenu("LevelEditor.ActorContextMenu");
