UVEditor: Fix missing plugin dependency

diff --git a/Engine/Plugins/Experimental/UVEditor/UVEditor.uplugin b/Engine/Plugins/Experimental/UVEditor/UVEditor.uplugin
--- a/Engine/Plugins/Experimental/UVEditor/UVEditor.uplugin
+++ b/Engine/Plugins/Experimental/UVEditor/UVEditor.uplugin
@@ -35,6 +35,11 @@
 			"Name": "GeometryProcessing",
 			"Enabled": true
 		},
+		{
+			"Name": "MeshModelingToolset",
+			"Enabled": true,
+			"TargetAllowList": [ "Editor" ]
+		},
 		{
 			"Name": "MeshModelingToolsetExp",
 			"Enabled": true,
