UVEditor: Fix issue where converting a selection resulting in an empty selection set would cause an ensure to be thrown.

diff --git a/Engine/Plugins/Experimental/UVEditor/Source/UVEditorTools/Private/Selection/UVEditorMeshSelectionMechanic.cpp b/Engine/Plugins/Experimental/UVEditor/Source/UVEditorTools/Private/Selection/UVEditorMeshSelectionMechanic.cpp
--- a/Engine/Plugins/Experimental/UVEditor/Source/UVEditorTools/Private/Selection/UVEditorMeshSelectionMechanic.cpp
+++ b/Engine/Plugins/Experimental/UVEditor/Source/UVEditorTools/Private/Selection/UVEditorMeshSelectionMechanic.cpp
@@ -987,6 +987,17 @@ void UUVEditorMeshSelectionMechanic::SetSelectionMode(
 		NewUnsetSelections.Add(OriginalUnsetSelection.GetConvertedSelection(*OriginalUnsetSelection.Target->AppliedCanonical, ExpectedSelectionType));
 	}
 
+	// Remove any selections that end up empty after conversion
+	NewSelections.RemoveAll([](const FUVToolSelection& Selection) {
+		return Selection.SelectedIDs.IsEmpty() && 
+			!(Selection.Type == FUVToolSelection::EType::Edge && Selection.HasStableEdgeIdentifiers());
+		});
+
+	NewUnsetSelections.RemoveAll([](const FUVToolSelection& Selection) {
+		return Selection.SelectedIDs.IsEmpty() &&
+			!(Selection.Type == FUVToolSelection::EType::Edge && Selection.HasStableEdgeIdentifiers());
+		});
+
 	// Apply selection change
 	SelectionAPI->BeginChange();
 	SelectionAPI->ClearSelections(false, false);
