UVEditor: Fix issue where Align and Distribute tools had incorrect text strings describing them and listed in the undo history.

diff --git a/Engine/Plugins/Experimental/UVEditor/Source/UVEditorTools/Private/UVEditorTransformTool.cpp b/Engine/Plugins/Experimental/UVEditor/Source/UVEditorTools/Private/UVEditorTransformTool.cpp
--- a/Engine/Plugins/Experimental/UVEditor/Source/UVEditorTools/Private/UVEditorTransformTool.cpp
+++ b/Engine/Plugins/Experimental/UVEditor/Source/UVEditorTools/Private/UVEditorTransformTool.cpp
@@ -20,6 +20,74 @@ using namespace UE::Geometry;
 
 #define LOCTEXT_NAMESPACE "UUVEditorTransformTool"
 
+namespace TransformToolLocals
+{
+	FText ToolName(EUVEditorUVTransformType Mode)
+	{
+		switch (Mode)
+		{
+		case EUVEditorUVTransformType::Transform:
+			return LOCTEXT("ToolName", "UV Transform");
+		case EUVEditorUVTransformType::Align:
+			return LOCTEXT("ToolName", "UV Align");
+		case EUVEditorUVTransformType::Distribute:
+			return LOCTEXT("ToolName", "UV Distribute");
+		default:
+			ensure(false);
+			return FText();
+		}
+	}
+
+	FText ToolDescription(EUVEditorUVTransformType Mode)
+	{
+		switch (Mode)
+		{
+		case EUVEditorUVTransformType::Transform:
+			return LOCTEXT("OnStartTool", "Translate, rotate or scale existing UV Charts using various strategies");
+		case EUVEditorUVTransformType::Align:
+			return LOCTEXT("OnStartTool", "Align UV elements relative to various positions and with various strategies");
+		case EUVEditorUVTransformType::Distribute:
+			return LOCTEXT("OnStartTool", "Distribute UV elements spatially with various strategies");
+		default:
+			ensure(false);
+			return FText();
+		}
+	}
+
+	FText ToolTransaction(EUVEditorUVTransformType Mode)
+	{
+		switch (Mode)
+		{
+		case EUVEditorUVTransformType::Transform:
+			return LOCTEXT("TransactionName", "Transform Tool");
+		case EUVEditorUVTransformType::Align:
+			return LOCTEXT("TransactionName", "Align Tool");
+		case EUVEditorUVTransformType::Distribute:
+			return LOCTEXT("TransactionName", "Distribute Tool");
+		default:
+			ensure(false);
+			return FText();
+		}
+	}
+
+	FText ToolConfirmation(EUVEditorUVTransformType Mode)
+	{
+		switch (Mode)
+		{
+		case EUVEditorUVTransformType::Transform:
+			return LOCTEXT("ApplyTool", "Transform Tool");
+		case EUVEditorUVTransformType::Align:
+			return LOCTEXT("ApplyTool", "Align Tool");
+		case EUVEditorUVTransformType::Distribute:
+			return LOCTEXT("ApplyTool", "Distribute Tool");
+		default:
+			ensure(false);
+			return FText();
+		}
+	}
+}
+
+
 // Tool builder
 // TODO: Could consider sharing some of the tool builder boilerplate for UV editor tools in a common base class.
 
@@ -153,8 +221,8 @@ void UUVEditorTransformTool::Setup()
 		}
 	}
 
-	SetToolDisplayName(LOCTEXT("ToolName", "UV Transform"));
-	GetToolManager()->DisplayMessage(LOCTEXT("OnStartUVTransformTool", "Translate, rotate or scale existing UV Charts using various strategies"),
+	SetToolDisplayName(TransformToolLocals::ToolName(ToolMode.Get(EUVEditorUVTransformType::Transform)));
+	GetToolManager()->DisplayMessage(TransformToolLocals::ToolDescription(ToolMode.Get(EUVEditorUVTransformType::Transform)),
 		EToolMessageLevel::UserNotification);
 
 	// Analytics
@@ -172,7 +240,7 @@ void UUVEditorTransformTool::Shutdown(EToolShutdownType ShutdownType)
 	if (ShutdownType == EToolShutdownType::Accept)
 	{
 		UUVToolEmitChangeAPI* ChangeAPI = GetToolManager()->GetContextObjectStore()->FindContext<UUVToolEmitChangeAPI>();
-		const FText TransactionName(LOCTEXT("TransformTransactionName", "Transform Tool"));
+		const FText TransactionName(TransformToolLocals::ToolTransaction(ToolMode.Get(EUVEditorUVTransformType::Transform)));
 		ChangeAPI->BeginUndoTransaction(TransactionName);
 
 		for (TObjectPtr<UUVEditorToolMeshInput> Target : Targets)
@@ -192,7 +260,8 @@ void UUVEditorTransformTool::Shutdown(EToolShutdownType ShutdownType)
 			// TODO: Again, it's not clear whether we need to update the entire triangle topology...
 			Target->UpdateCanonicalFromPreviews();
 
-			ChangeAPI->EmitToolIndependentUnwrapCanonicalChange(Target, ChangeTracker.EndChange(), LOCTEXT("ApplyTransformTool", "Transform Tool"));
+			ChangeAPI->EmitToolIndependentUnwrapCanonicalChange(Target, ChangeTracker.EndChange(),
+				       TransformToolLocals::ToolConfirmation(ToolMode.Get(EUVEditorUVTransformType::Transform)));
 		}
 
 		ChangeAPI->EndUndoTransaction();
