UVEditor: Fix issue where activating some tools resets the selection mode, causing the user to re-select their desired mode after finishing the tool.

diff --git a/Engine/Plugins/Experimental/UVEditor/Source/UVEditorTools/Private/UVEditorLayoutTool.cpp b/Engine/Plugins/Experimental/UVEditor/Source/UVEditorTools/Private/UVEditorLayoutTool.cpp
--- a/Engine/Plugins/Experimental/UVEditor/Source/UVEditorTools/Private/UVEditorLayoutTool.cpp
+++ b/Engine/Plugins/Experimental/UVEditor/Source/UVEditorTools/Private/UVEditorLayoutTool.cpp
@@ -54,11 +54,10 @@ void UUVEditorLayoutTool::Setup()
 	UUVToolSelectionAPI::FHighlightOptions HighlightOptions;
 	HighlightOptions.bBaseHighlightOnPreviews = true;
 	HighlightOptions.bAutoUpdateUnwrap = true;
-	UVToolSelectionAPI->SetSelectionMechanicMode(UUVToolSelectionAPI::EUVEditorSelectionMode::Triangle);
 	UVToolSelectionAPI->SetHighlightOptions(HighlightOptions);
 	UVToolSelectionAPI->SetHighlightVisible(true, false, true);
 
-	auto SetupOpFactory = [this](UUVEditorToolMeshInput& Target, TSet<int32>* Selection)
+	auto SetupOpFactory = [this](UUVEditorToolMeshInput& Target, const FUVToolSelection* Selection)
 	{
 		TObjectPtr<UUVEditorUVLayoutOperatorFactory> Factory = NewObject<UUVEditorUVLayoutOperatorFactory>();
 		Factory->TargetTransform = Target.AppliedPreview->PreviewMesh->GetTransform();
@@ -67,7 +66,8 @@ void UUVEditorLayoutTool::Setup()
 		Factory->GetSelectedUVChannel = [&Target]() { return Target.UVLayerIndex; };
 		if (Selection)
 		{
-			Factory->Selection.Emplace(*Selection);
+			Factory->Selection.Emplace(Selection->GetConvertedSelection(*Selection->Target->UnwrapCanonical,
+				                                                        UE::Geometry::FUVToolSelection::EType::Triangle).SelectedIDs);
 		}
 
 		Target.AppliedPreview->ChangeOpFactory(Factory);
@@ -86,10 +86,7 @@ void UUVEditorLayoutTool::Setup()
 		Factories.Reserve(UVToolSelectionAPI->GetSelections().Num());
 		for (FUVToolSelection Selection : UVToolSelectionAPI->GetSelections())
 		{
-			if (ensure(Selection.Type == FUVToolSelection::EType::Triangle))
-			{
-				Factories.Add(SetupOpFactory(*Selection.Target, &Selection.SelectedIDs));
-			}
+			Factories.Add(SetupOpFactory(*Selection.Target, &Selection));
 		}
 	}
 	else
diff --git a/Engine/Plugins/Experimental/UVEditor/Source/UVEditorTools/Private/UVEditorRecomputeUVsTool.cpp b/Engine/Plugins/Experimental/UVEditor/Source/UVEditorTools/Private/UVEditorRecomputeUVsTool.cpp
--- a/Engine/Plugins/Experimental/UVEditor/Source/UVEditorTools/Private/UVEditorRecomputeUVsTool.cpp
+++ b/Engine/Plugins/Experimental/UVEditor/Source/UVEditorTools/Private/UVEditorRecomputeUVsTool.cpp
@@ -61,7 +61,7 @@ void UUVEditorRecomputeUVsTool::Setup()
 	UUVToolSelectionAPI::FHighlightOptions HighlightOptions;
 	HighlightOptions.bBaseHighlightOnPreviews = true;
 	HighlightOptions.bAutoUpdateUnwrap = true;
-	UVToolSelectionAPI->SetSelectionMechanicMode(UUVToolSelectionAPI::EUVEditorSelectionMode::Triangle);
+
 	UVToolSelectionAPI->SetHighlightOptions(HighlightOptions);
 	UVToolSelectionAPI->SetHighlightVisible(true, false, true);
 
@@ -81,7 +81,7 @@ void UUVEditorRecomputeUVsTool::Setup()
 	}
 	UpdateActiveGroupLayer(false);  /* Don't update factories that don't exist yet. */
 
-	auto SetupOpFactory = [this](UUVEditorToolMeshInput& Target, TSet<int32>* Selection)
+	auto SetupOpFactory = [this](UUVEditorToolMeshInput& Target, const FUVToolSelection* Selection)
 	{
 		TObjectPtr<UUVEditorRecomputeUVsOpFactory> Factory = NewObject<UUVEditorRecomputeUVsOpFactory>();
 		Factory->TargetTransform = (FTransform3d)Target.AppliedPreview->PreviewMesh->GetTransform();
@@ -91,7 +91,8 @@ void UUVEditorRecomputeUVsTool::Setup()
 		Factory->GetSelectedUVChannel = [&Target]() { return Target.UVLayerIndex; };
 		if (Selection)
 		{
-			Factory->Selection.Emplace(*Selection);
+			Factory->Selection.Emplace(Selection->GetConvertedSelection(*Selection->Target->UnwrapCanonical,
+				                                                        UE::Geometry::FUVToolSelection::EType::Triangle).SelectedIDs);
 		}
 
 		Target.AppliedPreview->ChangeOpFactory(Factory);
@@ -113,10 +114,7 @@ void UUVEditorRecomputeUVsTool::Setup()
 		Factories.Reserve(UVToolSelectionAPI->GetSelections().Num());
 		for (FUVToolSelection Selection : UVToolSelectionAPI->GetSelections())
 		{
-			if (ensure(Selection.Type == FUVToolSelection::EType::Triangle))
-			{
-				Factories.Add(SetupOpFactory(*Selection.Target, &Selection.SelectedIDs));
-			}
+			Factories.Add(SetupOpFactory(*Selection.Target, &Selection));
 		}
 	}
 	else
