UVEditor: Removed ability for the FEditorViewport CameraController to afffect the 2D Unwrap viewport. This solves the issues around keyboard navigation causing unexpected lockups in the editor during extreme zoom instances. Also added upper and lower zoom limits, to prevent issues where crashes could happen at floating point limits.

diff --git a/Engine/Plugins/Experimental/UVEditor/Source/UVEditor/Private/UVEditor2DViewportBehaviorTargets.cpp b/Engine/Plugins/Experimental/UVEditor/Source/UVEditor/Private/UVEditor2DViewportBehaviorTargets.cpp
--- a/Engine/Plugins/Experimental/UVEditor/Source/UVEditor/Private/UVEditor2DViewportBehaviorTargets.cpp
+++ b/Engine/Plugins/Experimental/UVEditor/Source/UVEditor/Private/UVEditor2DViewportBehaviorTargets.cpp
@@ -95,7 +95,7 @@ void FUVEditor2DMouseWheelZoomBehaviorTarget::OnMouseWheelScrollUp(const FInputD
 
 	// Don't zoom in so far that the XY plane lies in front of our near clipping plane, or else everything
 	// will suddenly disappear.
-	if (NewLocation.Z > ViewportClient->GetNearClipPlane())
+	if (NewLocation.Z > ViewportClient->GetNearClipPlane() && NewLocation.Z > ZoomInLimit)
 	{
 		ViewportClient->SetViewLocation(NewLocation);
 	}
@@ -110,7 +110,10 @@ void FUVEditor2DMouseWheelZoomBehaviorTarget::OnMouseWheelScrollDown(const FInpu
 	ViewportClient->OverrideFarClipPlane(NewLocation.Z - CameraFarPlaneWorldZ);
 	ViewportClient->OverrideNearClipPlane(NewLocation.Z * (1.0 - CameraNearPlaneProportionZ));
 
-	ViewportClient->SetViewLocation(NewLocation);
+	if (NewLocation.Z < ZoomOutLimit)
+	{
+		ViewportClient->SetViewLocation(NewLocation);
+	}
 }
 
 void FUVEditor2DMouseWheelZoomBehaviorTarget::SetZoomAmount(double PercentZoomIn)
@@ -124,6 +127,12 @@ void FUVEditor2DMouseWheelZoomBehaviorTarget::SetZoomAmount(double PercentZoomIn
 	ZoomOutProportion = ZoomInProportion / (1 - ZoomInProportion);
 }
 
+void FUVEditor2DMouseWheelZoomBehaviorTarget::SetZoomLimits(double ZoomInLimitIn, double ZoomOutLimitIn)
+{
+	ZoomInLimit = ZoomInLimitIn;
+	ZoomOutLimit = ZoomOutLimitIn;
+}
+
 void FUVEditor2DMouseWheelZoomBehaviorTarget::SetCameraFarPlaneWorldZ(double CameraFarPlaneWorldZIn)
 {
 	CameraFarPlaneWorldZ = CameraFarPlaneWorldZIn;
diff --git a/Engine/Plugins/Experimental/UVEditor/Source/UVEditor/Private/UVEditor2DViewportClient.cpp b/Engine/Plugins/Experimental/UVEditor/Source/UVEditor/Private/UVEditor2DViewportClient.cpp
--- a/Engine/Plugins/Experimental/UVEditor/Source/UVEditor/Private/UVEditor2DViewportClient.cpp
+++ b/Engine/Plugins/Experimental/UVEditor/Source/UVEditor/Private/UVEditor2DViewportClient.cpp
@@ -16,6 +16,8 @@
 #include "Engine/Canvas.h"
 #include "Math/Box.h"
 #include "MathUtil.h"
+#include "CameraController.h"
+
 
 namespace FUVEditor2DViewportClientLocals {
 
@@ -159,6 +161,11 @@ FUVEditor2DViewportClient::FUVEditor2DViewportClient(FEditorModeTools* InModeToo
 
 	// Set up viewport manipulation behaviors:
 
+	FEditorCameraController* CameraControllerPtr = GetCameraController();
+	CameraController->GetConfig().MovementAccelerationRate = 0.0;
+	CameraController->GetConfig().RotationAccelerationRate = 0.0;
+	CameraController->GetConfig().FOVAccelerationRate = 0.0;
+
 	BehaviorSet = NewObject<UInputBehaviorSet>();
 
 	// We'll have the priority of our viewport manipulation behaviors be lower (i.e. higher
@@ -175,6 +182,7 @@ FUVEditor2DViewportClient::FUVEditor2DViewportClient(FEditorModeTools* InModeToo
 	ZoomBehaviorTarget = MakeUnique<FUVEditor2DMouseWheelZoomBehaviorTarget>(this);
 	ZoomBehaviorTarget->SetCameraFarPlaneWorldZ(FUVEditorUXSettings::CameraFarPlaneWorldZ);
 	ZoomBehaviorTarget->SetCameraNearPlaneProportionZ(FUVEditorUXSettings::CameraNearPlaneProportionZ);
+	ZoomBehaviorTarget->SetZoomLimits(0.001, 100000);
 	UMouseWheelInputBehavior* ZoomBehavior = NewObject<UMouseWheelInputBehavior>();
 	ZoomBehavior->Initialize(ZoomBehaviorTarget.Get());
 	ZoomBehavior->SetDefaultPriority(DEFAULT_VIEWPORT_BEHAVIOR_PRIORITY);	
diff --git a/Engine/Plugins/Experimental/UVEditor/Source/UVEditor/Public/UVEditor2DViewportBehaviorTargets.h b/Engine/Plugins/Experimental/UVEditor/Source/UVEditor/Public/UVEditor2DViewportBehaviorTargets.h
--- a/Engine/Plugins/Experimental/UVEditor/Source/UVEditor/Public/UVEditor2DViewportBehaviorTargets.h
+++ b/Engine/Plugins/Experimental/UVEditor/Source/UVEditor/Public/UVEditor2DViewportBehaviorTargets.h
@@ -54,6 +54,7 @@ public:
 	 *  back to where you were.
 	 */
 	virtual void SetZoomAmount(double PercentZoomIn);
+	void SetZoomLimits(double ZoomInLimitIn, double ZoomOutLimitIn);
 
 	void SetCameraFarPlaneWorldZ(double CameraFarPlaneWorldZIn);
 	void SetCameraNearPlaneProportionZ(double CameraFarPlaneProportionZIn );
@@ -64,6 +65,7 @@ protected:
 	double ZoomInProportion;
 	double ZoomOutProportion;
 	double ZoomInLimit;
+	double ZoomOutLimit;
 	double CameraFarPlaneWorldZ;
 	double CameraNearPlaneProportionZ;
 };
\ No newline at end of file
