UVEditor: Fix conflicting localization string keys for transform/align/distribute tool.

diff --git a/Engine/Plugins/Experimental/UVEditor/Source/UVEditorTools/Private/UVEditorTransformTool.cpp b/Engine/Plugins/Experimental/UVEditor/Source/UVEditorTools/Private/UVEditorTransformTool.cpp
--- a/Engine/Plugins/Experimental/UVEditor/Source/UVEditorTools/Private/UVEditorTransformTool.cpp
+++ b/Engine/Plugins/Experimental/UVEditor/Source/UVEditorTools/Private/UVEditorTransformTool.cpp
@@ -27,11 +27,11 @@ namespace TransformToolLocals
 		switch (Mode)
 		{
 		case EUVEditorUVTransformType::Transform:
-			return LOCTEXT("ToolName", "UV Transform");
+			return LOCTEXT("ToolNameTransform", "UV Transform");
 		case EUVEditorUVTransformType::Align:
-			return LOCTEXT("ToolName", "UV Align");
+			return LOCTEXT("ToolNameAlign", "UV Align");
 		case EUVEditorUVTransformType::Distribute:
-			return LOCTEXT("ToolName", "UV Distribute");
+			return LOCTEXT("ToolNameDistribute", "UV Distribute");
 		default:
 			ensure(false);
 			return FText();
@@ -43,11 +43,11 @@ namespace TransformToolLocals
 		switch (Mode)
 		{
 		case EUVEditorUVTransformType::Transform:
-			return LOCTEXT("OnStartTool", "Translate, rotate or scale existing UV Charts using various strategies");
+			return LOCTEXT("OnStartToolTransform", "Translate, rotate or scale existing UV Charts using various strategies");
 		case EUVEditorUVTransformType::Align:
-			return LOCTEXT("OnStartTool", "Align UV elements relative to various positions and with various strategies");
+			return LOCTEXT("OnStartToolAlign", "Align UV elements relative to various positions and with various strategies");
 		case EUVEditorUVTransformType::Distribute:
-			return LOCTEXT("OnStartTool", "Distribute UV elements spatially with various strategies");
+			return LOCTEXT("OnStartToolDistribute", "Distribute UV elements spatially with various strategies");
 		default:
 			ensure(false);
 			return FText();
@@ -59,11 +59,11 @@ namespace TransformToolLocals
 		switch (Mode)
 		{
 		case EUVEditorUVTransformType::Transform:
-			return LOCTEXT("TransactionName", "Transform Tool");
+			return LOCTEXT("TransactionNameTransform", "Transform Tool");
 		case EUVEditorUVTransformType::Align:
-			return LOCTEXT("TransactionName", "Align Tool");
+			return LOCTEXT("TransactionNameAlign", "Align Tool");
 		case EUVEditorUVTransformType::Distribute:
-			return LOCTEXT("TransactionName", "Distribute Tool");
+			return LOCTEXT("TransactionNameDistribute", "Distribute Tool");
 		default:
 			ensure(false);
 			return FText();
@@ -75,11 +75,11 @@ namespace TransformToolLocals
 		switch (Mode)
 		{
 		case EUVEditorUVTransformType::Transform:
-			return LOCTEXT("ApplyTool", "Transform Tool");
+			return LOCTEXT("ApplyToolTransform", "Transform Tool");
 		case EUVEditorUVTransformType::Align:
-			return LOCTEXT("ApplyTool", "Align Tool");
+			return LOCTEXT("ApplyToolAlign", "Align Tool");
 		case EUVEditorUVTransformType::Distribute:
-			return LOCTEXT("ApplyTool", "Distribute Tool");
+			return LOCTEXT("ApplyToolDistribute", "Distribute Tool");
 		default:
 			ensure(false);
 			return FText();
