UVEditor: Change from Experimental to Beta, and enable by default

diff --git a/Engine/Plugins/Experimental/UVEditor/UVEditor.uplugin b/Engine/Plugins/Experimental/UVEditor/UVEditor.uplugin
--- a/Engine/Plugins/Experimental/UVEditor/UVEditor.uplugin
+++ b/Engine/Plugins/Experimental/UVEditor/UVEditor.uplugin
@@ -1,9 +1,9 @@
 {
 	"FileVersion": 3,
-	"Version": 1,
-	"VersionName": "0.1",
+	"Version": 2,
+	"VersionName": "0.2",
 	"FriendlyName": "UVEditor",
-	"Description": "Asset editor for modifying the UV's of a mesh, built on the Interactive Tools Framework",
+	"Description": "Asset editor for modifying the UV mapping of a mesh",
 	"Category": "Other",
 	"CreatedBy": "Epic Games, Inc.",
 	"CreatedByURL": "http://epicgames.com",
@@ -11,8 +11,9 @@
 	"MarketplaceURL": "",
 	"SupportURL": "",
 	"CanContainContent": true,
-	"IsExperimentalVersion": true,
+	"IsBetaVersion": true,
 	"Installed": false,
+	"EnabledByDefault": true,
 	"Modules": [
 		{
 			"Name": "UVEditor",
