UVEditor: Adjust the behavior of the Transform rotation functions to rotate clockwise for positive rotation values and counterclockwise for negative values. Also addresses the jira issues of QuickRotate label placement and negative labels on CCW rotations.

diff --git a/Engine/Plugins/Experimental/UVEditor/Source/UVEditorTools/Private/Operators/UVEditorUVTransformOp.cpp b/Engine/Plugins/Experimental/UVEditor/Source/UVEditorTools/Private/Operators/UVEditorUVTransformOp.cpp
--- a/Engine/Plugins/Experimental/UVEditor/Source/UVEditorTools/Private/Operators/UVEditorUVTransformOp.cpp
+++ b/Engine/Plugins/Experimental/UVEditor/Source/UVEditorTools/Private/Operators/UVEditorUVTransformOp.cpp
@@ -298,7 +298,9 @@ void FUVEditorUVTransformOp::HandleTransformationOp(FProgressCancel * Progress)
 	{
 		FVector2f UV = TransformOpLocals::UnwrapPositionToUV(ResultMesh->GetVertexRef(Vid));
 		FVector2f UV_Rotated;
-		double RotationInRadians = RotationIn / 180.0 * UE_PI;
+		// We are flipping the sign here to match the conventions in other UV editors
+		// where positive values are clockwise and negative values are counterclockwise.
+		double RotationInRadians = -RotationIn / 180.0 * UE_PI;
 		UV = (UV - Pivot);
 		UV_Rotated[0] = UV[0] * FMath::Cos(RotationInRadians) - UV[1] * FMath::Sin(RotationInRadians);
 		UV_Rotated[1] = UV[0] * FMath::Sin(RotationInRadians) + UV[1] * FMath::Cos(RotationInRadians);
diff --git a/Engine/Plugins/Experimental/UVEditor/Source/UVEditorToolsEditorOnly/Private/DetailsCustomizations/UVTransformToolCustomizations.cpp b/Engine/Plugins/Experimental/UVEditor/Source/UVEditorToolsEditorOnly/Private/DetailsCustomizations/UVTransformToolCustomizations.cpp
--- a/Engine/Plugins/Experimental/UVEditor/Source/UVEditorToolsEditorOnly/Private/DetailsCustomizations/UVTransformToolCustomizations.cpp
+++ b/Engine/Plugins/Experimental/UVEditor/Source/UVEditorToolsEditorOnly/Private/DetailsCustomizations/UVTransformToolCustomizations.cpp
@@ -257,7 +257,7 @@ void FUVEditorUVTransformToolDetails::BuildQuickRotateMenu(IDetailLayoutBuilder&
 	[
 		SNew(SUniformGridPanel)
 		.SlotPadding(FMargin(5.0f))
-		+ SUniformGridPanel::Slot(0, 0)
+		+ SUniformGridPanel::Slot(2, 0)
 		.HAlign(HAlign_Center)
 		.VAlign(VAlign_Center)
 		[
@@ -267,7 +267,7 @@ void FUVEditorUVTransformToolDetails::BuildQuickRotateMenu(IDetailLayoutBuilder&
 			.OnClicked_Lambda([this, QuickRotationOffsetHandle, ApplyRotation]() {
 				float RotationValue;
 				QuickRotationOffsetHandle->GetValue(RotationValue);
-				return ApplyRotation(-RotationValue);
+				return ApplyRotation(RotationValue);
 			})
 		]
 		+ SUniformGridPanel::Slot(1, 0)
@@ -277,7 +277,7 @@ void FUVEditorUVTransformToolDetails::BuildQuickRotateMenu(IDetailLayoutBuilder&
 			SNew(SProperty, QuickRotationOffsetHandle)
 			.ShouldDisplayName(false)
 		]
-		+ SUniformGridPanel::Slot(2, 0)
+		+ SUniformGridPanel::Slot(0, 0)
 		.HAlign(HAlign_Center)
 		.VAlign(VAlign_Center)
 		[
@@ -287,63 +287,63 @@ void FUVEditorUVTransformToolDetails::BuildQuickRotateMenu(IDetailLayoutBuilder&
 			.OnClicked_Lambda([this, QuickRotationOffsetHandle, ApplyRotation]() {
 				float RotationValue;
 				QuickRotationOffsetHandle->GetValue(RotationValue);
-				return ApplyRotation(RotationValue);
+				return ApplyRotation(-RotationValue);
 			})
 		]
 
-		+ SUniformGridPanel::Slot(0, 1)
+		+ SUniformGridPanel::Slot(2, 1)
 		.HAlign(HAlign_Center)
 		.VAlign(VAlign_Center)
 		[
 			SNew(SButton)
 			.Text(LOCTEXT("QuickRotateClockwise10Deg", "10°"))
 			.ToolTipText(LOCTEXT("QuickRotateClockwise10DegToolTip", "Applies a 10 degree clockwise orientation"))
-			.OnClicked_Lambda([this, ApplyRotation]() { return ApplyRotation(-10.0); })
+			.OnClicked_Lambda([this, ApplyRotation]() { return ApplyRotation(10.0); })
 		]
-		+ SUniformGridPanel::Slot(2, 1)
+		+ SUniformGridPanel::Slot(0, 1)
 		.HAlign(HAlign_Center)
 		.VAlign(VAlign_Center)
 		[
 			SNew(SButton)
-			.Text(LOCTEXT("QuickRotateCounterclockwise10Deg", "10°"))
+			.Text(LOCTEXT("QuickRotateCounterclockwise10Deg", "-10°"))
 			.ToolTipText(LOCTEXT("QuickRotateCounterclockwise10DegToolTip", "Applies a 10 degree counter clockwise orientation"))
-			.OnClicked_Lambda([this, ApplyRotation]() { return ApplyRotation(10.0); })
+			.OnClicked_Lambda([this, ApplyRotation]() { return ApplyRotation(-10.0); })
 		]
-		+ SUniformGridPanel::Slot(0, 2)
+		+ SUniformGridPanel::Slot(2, 2)
 		.HAlign(HAlign_Center)
 		.VAlign(VAlign_Center)
 		[
 			SNew(SButton)
 			.Text(LOCTEXT("QuickRotateClockwise45Deg", "45°"))
 			.ToolTipText(LOCTEXT("QuickRotateClockwise45DegToolTip", "Applies a 45 degree clockwise orientation"))
-			.OnClicked_Lambda([this, ApplyRotation]() { return ApplyRotation(-45.0); })
+			.OnClicked_Lambda([this, ApplyRotation]() { return ApplyRotation(45.0); })
 		]
-		+ SUniformGridPanel::Slot(2, 2)
+		+ SUniformGridPanel::Slot(0, 2)
 		.HAlign(HAlign_Center)
 		.VAlign(VAlign_Center)
 		[
 			SNew(SButton)
-			.Text(LOCTEXT("QuickRotateCounterclockwise45Deg", "45°"))
+			.Text(LOCTEXT("QuickRotateCounterclockwise45Deg", "-45°"))
 			.ToolTipText(LOCTEXT("QuickRotateCounterclockwise45DegToolTip", "Applies a 45 degree counter clockwise orientation"))
-			.OnClicked_Lambda([this, ApplyRotation]() { return ApplyRotation(45.0); })
+			.OnClicked_Lambda([this, ApplyRotation]() { return ApplyRotation(-45.0); })
 		]
-		+ SUniformGridPanel::Slot(0, 3)
+		+ SUniformGridPanel::Slot(2, 3)
 		.HAlign(HAlign_Center)
 		.VAlign(VAlign_Center)
 		[
 			SNew(SButton)
 			.Text(LOCTEXT("QuickRotateClockwise90Deg", "90°"))
 			.ToolTipText(LOCTEXT("QuickRotateClockwise90DegToolTip", "Applies a 90 degree clockwise orientation"))
-			.OnClicked_Lambda([this, ApplyRotation]() { return ApplyRotation(-90.0); })
+			.OnClicked_Lambda([this, ApplyRotation]() { return ApplyRotation(90.0); })
 		]
-		+ SUniformGridPanel::Slot(2, 3)
+		+ SUniformGridPanel::Slot(0, 3)
 		.HAlign(HAlign_Center)
 		.VAlign(VAlign_Center)
 		[
 			SNew(SButton)
-			.Text(LOCTEXT("QuickRotateCounterclockwise90Deg", "90°"))
+			.Text(LOCTEXT("QuickRotateCounterclockwise90Deg", "-90°"))
 			.ToolTipText(LOCTEXT("QuickRotateCounterclockwise90DegToolTip", "Applies a 90 degree counter clockwise orientation"))
-			.OnClicked_Lambda([this, ApplyRotation]() { return ApplyRotation(90.0); })
+			.OnClicked_Lambda([this, ApplyRotation]() { return ApplyRotation(-90.0); })
 		]
 	];
 }
