UVEditor: Correctly handle the initiation and termination of Play In Editor (PIE) mode in the UV Editor. This properly restarts the selection tool upon finishing PIE and adds new warning messages indicating that PIE is currently running.

diff --git a/Engine/Plugins/Experimental/UVEditor/Source/UVEditor/Private/SUVEditor2DViewportToolBar.cpp b/Engine/Plugins/Experimental/UVEditor/Source/UVEditor/Private/SUVEditor2DViewportToolBar.cpp
--- a/Engine/Plugins/Experimental/UVEditor/Source/UVEditor/Private/SUVEditor2DViewportToolBar.cpp
+++ b/Engine/Plugins/Experimental/UVEditor/Source/UVEditor/Private/SUVEditor2DViewportToolBar.cpp
@@ -284,7 +284,9 @@ TSharedRef<SWidget> SUVEditor2DViewportToolBar::MakeTransformToolBar(const TShar
 	}
 	ToolbarBuilder.EndSection();
 
-	return ToolbarBuilder.MakeWidget();
+	TSharedRef<SWidget> TransformBar = ToolbarBuilder.MakeWidget();
+	TransformBar->SetEnabled(TAttribute<bool>::CreateLambda([this]() {return Viewport2DClient->AreWidgetButtonsEnabled(); }));
+	return TransformBar;
 }
 
 // The following methods again mimic the patterns found in the STransformViewportToolbar.cpp,
diff --git a/Engine/Plugins/Experimental/UVEditor/Source/UVEditor/Private/UVEditorMode.cpp b/Engine/Plugins/Experimental/UVEditor/Source/UVEditor/Private/UVEditorMode.cpp
--- a/Engine/Plugins/Experimental/UVEditor/Source/UVEditor/Private/UVEditorMode.cpp
+++ b/Engine/Plugins/Experimental/UVEditor/Source/UVEditor/Private/UVEditorMode.cpp
@@ -264,6 +264,23 @@ void UUVEditorMode::Enter()
 {
 	Super::Enter();
 
+	BeginPIEDelegateHandle = FEditorDelegates::PreBeginPIE.AddLambda([this](bool bSimulating)
+	{	
+		SetSimulationWarning(true);
+	});
+
+	EndPIEDelegateHandle = FEditorDelegates::EndPIE.AddLambda([this](bool bSimulating)
+	{
+		ActivateDefaultTool();
+		SetSimulationWarning(false);
+	});
+
+	CancelPIEDelegateHandle = FEditorDelegates::CancelPIE.AddLambda([this]()
+	{
+		ActivateDefaultTool();
+		SetSimulationWarning(false);
+	});
+
 	// Our mode needs to get Render and DrawHUD calls, but doesn't implement the legacy interface that
 	// causes those to be called. Instead, we'll hook into the tools context's Render and DrawHUD calls.
 	GetInteractiveToolsContext()->OnRender.AddUObject(this, &UUVEditorMode::Render);
@@ -766,6 +783,10 @@ void UUVEditorMode::Exit()
 		LivePreviewITC->OnDrawHUD.RemoveAll(this);
 	}
 
+	FEditorDelegates::PreBeginPIE.Remove(BeginPIEDelegateHandle);
+	FEditorDelegates::EndPIE.Remove(EndPIEDelegateHandle);
+	FEditorDelegates::CancelPIE.Remove(CancelPIEDelegateHandle);
+
 	Super::Exit();
 }
 
@@ -1263,6 +1284,15 @@ void UUVEditorMode::ModeTick(float DeltaTime)
 
 }
 
+void UUVEditorMode::SetSimulationWarning(bool bEnabled)
+{
+	if (Toolkit.IsValid())
+	{
+		FUVEditorModeToolkit* UVEditorModeToolkit = StaticCast<FUVEditorModeToolkit*>(Toolkit.Get());
+		UVEditorModeToolkit->EnableShowPIEWarning(bEnabled);
+	}
+}
+
 int32 UUVEditorMode::GetNumUVChannels(int32 AssetID) const
 {
 	if (AssetID < 0 || AssetID >= AppliedCanonicalMeshes.Num())
diff --git a/Engine/Plugins/Experimental/UVEditor/Source/UVEditor/Private/UVEditorModeToolkit.cpp b/Engine/Plugins/Experimental/UVEditor/Source/UVEditor/Private/UVEditorModeToolkit.cpp
--- a/Engine/Plugins/Experimental/UVEditor/Source/UVEditor/Private/UVEditorModeToolkit.cpp
+++ b/Engine/Plugins/Experimental/UVEditor/Source/UVEditor/Private/UVEditorModeToolkit.cpp
@@ -47,7 +47,16 @@ FUVEditorModeToolkit::FUVEditorModeToolkit()
 		.Text(FText::GetEmpty())
 		.Visibility(EVisibility::Collapsed)
 	]
-
+	+ SVerticalBox::Slot()
+	.AutoHeight()
+	[
+		SAssignNew(ModeWarningArea, STextBlock)
+		.AutoWrapText(true)
+		.Font(FCoreStyle::GetDefaultFontStyle("Bold", 9))
+		.ColorAndOpacity(FSlateColor(FLinearColor(0.9f, 0.15f, 0.15f))) //TODO: This probably needs to not be hardcoded
+		.Text(FText::GetEmpty())
+		.Visibility(EVisibility::Collapsed)
+	]
 	+ SVerticalBox::Slot()
 	[
 		SAssignNew(ToolDetailsContainer, SBorder)
@@ -432,5 +441,44 @@ void FUVEditorModeToolkit::ClearWarning()
 	ToolWarningArea->SetVisibility(EVisibility::Collapsed);
 }
 
+void FUVEditorModeToolkit::PostModeWarning(const FText& Message)
+{
+	if (Message.IsEmpty())
+	{
+		ClearModeWarning();
+	}
+	else
+	{
+		ModeWarningArea->SetText(Message);
+		ModeWarningArea->SetVisibility(EVisibility::Visible);
+	}
+}
+
+void FUVEditorModeToolkit::ClearModeWarning()
+{
+	ModeWarningArea->SetText(FText());
+	ModeWarningArea->SetVisibility(EVisibility::Collapsed);
+}
+
+void FUVEditorModeToolkit::UpdatePIEWarnings()
+{
+	if (bShowPIEWarning)
+	{
+		PostModeWarning(LOCTEXT("ModelingModeToolkitPIEWarning", "UV Editor functionality is limited during Play in Editor sessions. End the current Play in Editor session to continue using the editor."));
+	}
+	else
+	{
+		ClearModeWarning();
+	}
+}
+
+void FUVEditorModeToolkit::EnableShowPIEWarning(bool bEnable)
+{
+	if (bShowPIEWarning != bEnable)
+	{
+		bShowPIEWarning = bEnable;
+		UpdatePIEWarnings();
+	}
+}
 
 #undef LOCTEXT_NAMESPACE
diff --git a/Engine/Plugins/Experimental/UVEditor/Source/UVEditor/Public/UVEditorMode.h b/Engine/Plugins/Experimental/UVEditor/Source/UVEditor/Public/UVEditorMode.h
--- a/Engine/Plugins/Experimental/UVEditor/Source/UVEditor/Public/UVEditorMode.h
+++ b/Engine/Plugins/Experimental/UVEditor/Source/UVEditor/Public/UVEditorMode.h
@@ -254,6 +254,7 @@ protected:
 	void InitializeTargets();
 	void RegisterTools();
 	void RegisterActions();
+	void SetSimulationWarning(bool bEnabled);
 
 	// UEdMode overrides
 	virtual void CreateToolkit() override;
@@ -353,5 +354,10 @@ protected:
 	FString DefaultToolIdentifier;
 
 	static FDateTime AnalyticsLastStartTimestamp;
+
+	// Holds references to PIE callbacks to handle logic when the PIE session starts & shuts down
+	FDelegateHandle BeginPIEDelegateHandle;
+	FDelegateHandle EndPIEDelegateHandle;
+	FDelegateHandle CancelPIEDelegateHandle;
 };
 
diff --git a/Engine/Plugins/Experimental/UVEditor/Source/UVEditor/Public/UVEditorModeToolkit.h b/Engine/Plugins/Experimental/UVEditor/Source/UVEditor/Public/UVEditorModeToolkit.h
--- a/Engine/Plugins/Experimental/UVEditor/Source/UVEditor/Public/UVEditorModeToolkit.h
+++ b/Engine/Plugins/Experimental/UVEditor/Source/UVEditor/Public/UVEditorModeToolkit.h
@@ -35,6 +35,9 @@ public:
 	/** Creates a widget where the visible UDIMs can be changed. */
 	virtual TSharedRef<SWidget> CreateUDIMSettingsWidget();
 
+	/** Enables or diables the Play in Editor mode warning for the editor */
+	virtual void EnableShowPIEWarning(bool bEnable);
+
 	// FModeToolkit
 	virtual void Init(const TSharedPtr<IToolkitHost>& InitToolkitHost, TWeakObjectPtr<UEdMode> InOwningMode) override;
 	virtual FText GetActiveToolDisplayName() const override { return ActiveToolName; }
@@ -61,10 +64,15 @@ public:
 	virtual void PostWarning(const FText& Message);
 	virtual void ClearWarning();
 
+	// set/clear mode-level warning message area. These messages don't get cleared on tool shutdown.
+	virtual void PostModeWarning(const FText& Message);
+	virtual void ClearModeWarning();
+
 protected:
 	const static TArray<FName> PaletteNames_Standard;
 
 	TSharedRef<SWidget> CreateDisplaySettingsWidget(UObject* SettingsObject) const;
+	void UpdatePIEWarnings();
 
 	// FModeToolkit
 	virtual void OnToolStarted(UInteractiveToolManager* Manager, UInteractiveTool* Tool) override;
@@ -80,6 +88,9 @@ protected:
 	// A place for tools to write out any warnings
 	TSharedPtr<STextBlock> ToolWarningArea;
 
+	// A place for the mode to write out any warnings
+	TSharedPtr<STextBlock> ModeWarningArea;
+
 	// A container for the tool settings that is populated by the DetailsView managed
 	// in FModeToolkit
 	TSharedPtr<SBorder> ToolDetailsContainer;
@@ -93,6 +104,7 @@ protected:
 	FText ActiveToolName;
 	const FSlateBrush* ActiveToolIcon;
 	FStatusBarMessageHandle ActiveToolMessageHandle;
+	bool bShowPIEWarning;
 
 	void UpdateActiveToolProperties();
 	void InvalidateCachedDetailPanelState(UObject* ChangedObject);
