Add an FChange transaction to track the UV editor's asset IDs across undo/redo of an Apply() action

diff --git a/Engine/Plugins/Experimental/UVEditor/Source/UVEditor/Private/UVEditorMode.cpp b/Engine/Plugins/Experimental/UVEditor/Source/UVEditor/Private/UVEditorMode.cpp
--- a/Engine/Plugins/Experimental/UVEditor/Source/UVEditor/Private/UVEditorMode.cpp
+++ b/Engine/Plugins/Experimental/UVEditor/Source/UVEditor/Private/UVEditorMode.cpp
@@ -173,6 +173,48 @@ namespace UVEditorModeLocals
 	}
 }
 
+
+namespace UVEditorModeChange
+{
+	/**
+	 * Change for undoing/redoing "Apply" of UV editor changes.
+	 */
+	class FApplyChangesChange : public FToolCommandChange
+	{
+	public:
+		FApplyChangesChange(TSet<int32> OriginalModifiedAssetIDsIn)
+			: OriginalModifiedAssetIDs(OriginalModifiedAssetIDsIn)
+		{
+		}
+
+		virtual void Apply(UObject* Object) override
+		{
+			UUVEditorMode* UVEditorMode = Cast<UUVEditorMode>(Object);
+			UVEditorMode->ModifiedAssetIDs.Reset();
+		}
+
+		virtual void Revert(UObject* Object) override
+		{
+			UUVEditorMode* UVEditorMode = Cast<UUVEditorMode>(Object);
+			UVEditorMode->ModifiedAssetIDs = OriginalModifiedAssetIDs;
+		}
+
+		virtual bool HasExpired(UObject* Object) const override
+		{
+			UUVEditorMode* UVEditorMode = Cast<UUVEditorMode>(Object);
+			return !(UVEditorMode && UVEditorMode->IsActive());
+		}
+
+		virtual FString ToString() const override
+		{
+			return TEXT("UVEditorModeLocals::FApplyChangesChange");
+		}
+
+	protected:
+		TSet<int32> OriginalModifiedAssetIDs;
+	};
+}
+
 const FToolTargetTypeRequirements& UUVEditorMode::GetToolTargetRequirements()
 {
 	static const FToolTargetTypeRequirements ToolTargetRequirements =
@@ -1046,7 +1088,8 @@ void UUVEditorMode::ApplyChanges()
 {
 	using namespace UVEditorModeLocals;
 
-	GetToolManager()->BeginUndoTransaction(LOCTEXT("UVEditorApplyChangesTransaction", "UV Editor Apply Changes"));
+	FText ApplyChangesText = LOCTEXT("UVEditorApplyChangesTransaction", "UV Editor Apply Changes");
+	GetToolManager()->BeginUndoTransaction(ApplyChangesText);
 
 	for (int32 AssetID : ModifiedAssetIDs)
 	{
@@ -1054,7 +1097,10 @@ void UUVEditorMode::ApplyChanges()
 		UE::ToolTarget::CommitDynamicMeshUVUpdate(ToolTargets[AssetID], AppliedCanonicalMeshes[AssetID].Get());
 	}
 
+	GetInteractiveToolsContext()->GetTransactionAPI()->AppendChange(
+		this, MakeUnique<UVEditorModeChange::FApplyChangesChange>(ModifiedAssetIDs), ApplyChangesText);
 	ModifiedAssetIDs.Reset();
+	GetInteractiveToolsContext()->GetTransactionAPI()->EndUndoTransaction();
 
 	GetToolManager()->EndUndoTransaction();
 }
diff --git a/Engine/Plugins/Experimental/UVEditor/Source/UVEditor/Public/UVEditorMode.h b/Engine/Plugins/Experimental/UVEditor/Source/UVEditor/Public/UVEditorMode.h
--- a/Engine/Plugins/Experimental/UVEditor/Source/UVEditor/Public/UVEditorMode.h
+++ b/Engine/Plugins/Experimental/UVEditor/Source/UVEditor/Public/UVEditorMode.h
@@ -132,6 +132,11 @@ private:
 	TArray<FString> UVAssetNames;
 };
 
+namespace UVEditorModeChange
+{
+	class FApplyChangesChange;
+}
+
 /**
  * The UV editor mode is the mode used in the UV asset editor. It holds most of the inter-tool state.
  * We put things in a mode instead of directly into the asset editor in case we want to someday use the mode
@@ -141,6 +146,8 @@ UCLASS(Transient)
 class UUVEditorMode : public UEdMode
 {
 	GENERATED_BODY()
+
+	friend UVEditorModeChange::FApplyChangesChange;
 public:
 	const static FEditorModeID EM_UVEditorModeId;
 
