UVEditor: Fix additional issues where starting Play in Editor (PIE) mode would allow the UV Editor to perform unwanted actions, such as saving assets.

diff --git a/Engine/Plugins/Experimental/UVEditor/Source/UVEditor/Private/UVEditorMode.cpp b/Engine/Plugins/Experimental/UVEditor/Source/UVEditor/Private/UVEditorMode.cpp
--- a/Engine/Plugins/Experimental/UVEditor/Source/UVEditor/Private/UVEditorMode.cpp
+++ b/Engine/Plugins/Experimental/UVEditor/Source/UVEditor/Private/UVEditorMode.cpp
@@ -306,19 +306,29 @@ void UUVEditorMode::Enter()
 {
 	Super::Enter();
 
-	BeginPIEDelegateHandle = FEditorDelegates::PreBeginPIE.AddLambda([this](bool bSimulating)
-	{	
+	bPIEModeActive = false;
+	if (GEditor->PlayWorld != NULL || GEditor->bIsSimulatingInEditor)
+	{
+		bPIEModeActive = true;
 		SetSimulationWarning(true);
+	}
+
+	BeginPIEDelegateHandle = FEditorDelegates::PreBeginPIE.AddLambda([this](bool bSimulating)
+	{
+		bPIEModeActive = true;
+        SetSimulationWarning(true);
 	});
 
 	EndPIEDelegateHandle = FEditorDelegates::EndPIE.AddLambda([this](bool bSimulating)
 	{
+		bPIEModeActive = false;
 		ActivateDefaultTool();
 		SetSimulationWarning(false);
 	});
 
 	CancelPIEDelegateHandle = FEditorDelegates::CancelPIE.AddLambda([this]()
 	{
+		bPIEModeActive = false;
 		ActivateDefaultTool();
 		SetSimulationWarning(false);
 	});
@@ -533,6 +543,11 @@ void UUVEditorMode::RegisterActions()
 
 bool UUVEditorMode::ShouldToolStartBeAllowed(const FString& ToolIdentifier) const
 {
+	if (bPIEModeActive)
+	{
+		return false;
+	}
+
 	// For now we've decided to disallow switch-away on accept/cancel tools in the UV editor.
 	if (GetInteractiveToolsContext()->ActiveToolHasAccept())
 	{
@@ -644,7 +659,10 @@ void UUVEditorMode::DrawHUD(FCanvas* Canvas, IToolsContextRenderAPI* RenderAPI)
 
 void UUVEditorMode::ActivateDefaultTool()
 {
-	GetInteractiveToolsContext()->StartTool(DefaultToolIdentifier);
+	if (!bPIEModeActive)
+	{
+		GetInteractiveToolsContext()->StartTool(DefaultToolIdentifier);
+	}
 }
 
 bool UUVEditorMode::IsDefaultToolActive()
@@ -1070,11 +1088,17 @@ void UUVEditorMode::EmitToolIndependentObjectChange(UObject* TargetObject, TUniq
 	GetInteractiveToolsContext()->GetTransactionAPI()->AppendChange(TargetObject, MoveTemp(Change), Description);
 }
 
-bool UUVEditorMode::HaveUnappliedChanges()
+bool UUVEditorMode::HaveUnappliedChanges() const
 {
 	return ModifiedAssetIDs.Num() > 0;
 }
 
+bool UUVEditorMode::CanApplyChanges() const
+{
+	return !bPIEModeActive && HaveUnappliedChanges();
+}
+
+
 void UUVEditorMode::GetAssetsWithUnappliedChanges(TArray<TObjectPtr<UObject>> UnappliedAssetsOut)
 {
 	for (int32 AssetID : ModifiedAssetIDs)
diff --git a/Engine/Plugins/Experimental/UVEditor/Source/UVEditor/Private/UVEditorToolkit.cpp b/Engine/Plugins/Experimental/UVEditor/Source/UVEditor/Private/UVEditorToolkit.cpp
--- a/Engine/Plugins/Experimental/UVEditor/Source/UVEditor/Private/UVEditorToolkit.cpp
+++ b/Engine/Plugins/Experimental/UVEditor/Source/UVEditor/Private/UVEditorToolkit.cpp
@@ -229,8 +229,29 @@ bool FUVEditorToolkit::OnRequestClose()
 		return true; 
 	}
 
+	bool bHasUnappliedChanges = UVMode->HaveUnappliedChanges();
+	bool bCanApplyChanges = UVMode->CanApplyChanges();
+
+	// Warn the user if there are unapplied changes *but* we can't currently save them.
+	if (bHasUnappliedChanges && !bCanApplyChanges)
+	{
+		EAppReturnType::Type YesNoReply = FMessageDialog::Open(EAppMsgType::YesNo,
+			NSLOCTEXT("UVEditor", "Prompt_UVEditorCloseCannotSave", "At least one of the assets has unapplied changes, however the UV Editor cannot currently apply changes due to current editor conditions. Do you still want to exit the UV Editor? (Selecting 'Yes' will cause all yet unapplied changes to be lost!)"));
+
+		switch (YesNoReply)
+		{
+		case EAppReturnType::Yes:
+			// exit without applying changes
+			break;
+
+		case EAppReturnType::No:
+			// don't exit
+			return false;
+		}
+	}
+
 	// Warn the user of any unapplied changes.
-	if (UVMode->HaveUnappliedChanges())
+	if (bHasUnappliedChanges && bCanApplyChanges)
 	{
 		TArray<TObjectPtr<UObject>> UnappliedAssets;
 		UVMode->GetAssetsWithUnappliedChanges(UnappliedAssets);
@@ -289,6 +310,19 @@ void FUVEditorToolkit::SaveAsset_Execute()
 	FAssetEditorToolkit::SaveAsset_Execute();
 }
 
+bool FUVEditorToolkit::CanSaveAsset() const 
+{
+	UUVEditorMode* UVMode = Cast<UUVEditorMode>(EditorModeManager->GetActiveScriptableMode(UUVEditorMode::EM_UVEditorModeId));
+	check(UVMode);
+	return UVMode->CanApplyChanges();
+}
+
+bool FUVEditorToolkit::CanSaveAssetAs() const 
+{
+	return CanSaveAsset();
+}
+
+
 TSharedRef<SDockTab> FUVEditorToolkit::SpawnTab_LivePreview(const FSpawnTabArgs& Args)
 {
 	TSharedRef< SDockTab > DockableTab =
@@ -426,7 +460,7 @@ void FUVEditorToolkit::PostInitAssetEditor()
 	ToolkitCommands->MapAction(
 		FUVEditorCommands::Get().ApplyChanges,
 		FExecuteAction::CreateUObject(UVMode, &UUVEditorMode::ApplyChanges),
-		FCanExecuteAction::CreateUObject(UVMode, &UUVEditorMode::HaveUnappliedChanges));
+		FCanExecuteAction::CreateUObject(UVMode, &UUVEditorMode::CanApplyChanges));
 	FName ParentToolbarName;
 	const FName ToolBarName = GetToolMenuToolbarName(ParentToolbarName);
 	UToolMenu* AssetToolbar = UToolMenus::Get()->ExtendMenu(ToolBarName);
diff --git a/Engine/Plugins/Experimental/UVEditor/Source/UVEditor/Public/UVEditorMode.h b/Engine/Plugins/Experimental/UVEditor/Source/UVEditor/Public/UVEditorMode.h
--- a/Engine/Plugins/Experimental/UVEditor/Source/UVEditor/Public/UVEditorMode.h
+++ b/Engine/Plugins/Experimental/UVEditor/Source/UVEditor/Public/UVEditorMode.h
@@ -195,7 +195,8 @@ public:
 	void EmitToolIndependentObjectChange(UObject* TargetObject, TUniquePtr<FToolCommandChange> Change, const FText& Description);
 
 	// Asset management
-	bool HaveUnappliedChanges();
+	bool HaveUnappliedChanges() const;
+	bool CanApplyChanges() const;
 	void GetAssetsWithUnappliedChanges(TArray<TObjectPtr<UObject>> UnappliedAssetsOut);
 	void ApplyChanges();
 
@@ -363,6 +364,7 @@ protected:
 	static FDateTime AnalyticsLastStartTimestamp;
 
 	// Holds references to PIE callbacks to handle logic when the PIE session starts & shuts down
+	bool bPIEModeActive;
 	FDelegateHandle BeginPIEDelegateHandle;
 	FDelegateHandle EndPIEDelegateHandle;
 	FDelegateHandle CancelPIEDelegateHandle;
diff --git a/Engine/Plugins/Experimental/UVEditor/Source/UVEditor/Public/UVEditorToolkit.h b/Engine/Plugins/Experimental/UVEditor/Source/UVEditor/Public/UVEditorToolkit.h
--- a/Engine/Plugins/Experimental/UVEditor/Source/UVEditor/Public/UVEditorToolkit.h
+++ b/Engine/Plugins/Experimental/UVEditor/Source/UVEditor/Public/UVEditorToolkit.h
@@ -51,6 +51,8 @@ public:
 	virtual void RegisterTabSpawners(const TSharedRef<FTabManager>& InTabManager) override;
 	virtual bool OnRequestClose() override;
 	virtual void SaveAsset_Execute() override;
+	virtual bool CanSaveAsset() const override;
+	virtual bool CanSaveAssetAs() const override;
 	void OnToolkitHostingStarted(const TSharedRef<IToolkit>& Toolkit) override;
 	void OnToolkitHostingFinished(const TSharedRef<IToolkit>& Toolkit) override;
 
