UVEditor: Fix crash when opening Transform/Align/Distribute tool with no selection and a UV channel other than 0 active.

diff --git a/Engine/Plugins/Experimental/UVEditor/Source/UVEditorTools/Private/Operators/UVEditorUVTransformOp.cpp b/Engine/Plugins/Experimental/UVEditor/Source/UVEditorTools/Private/Operators/UVEditorUVTransformOp.cpp
--- a/Engine/Plugins/Experimental/UVEditor/Source/UVEditorTools/Private/Operators/UVEditorUVTransformOp.cpp
+++ b/Engine/Plugins/Experimental/UVEditor/Source/UVEditorTools/Private/Operators/UVEditorUVTransformOp.cpp
@@ -202,6 +202,7 @@ void FUVEditorUVTransformBaseOp::CalculateResult(FProgressCancel* Progress)
 
 	int UVLayerInput = UVLayerIndex;
 	ActiveUVLayer = ResultMesh->Attributes()->GetUVLayer(UVLayerInput);
+	ensure(ActiveUVLayer);
 
 	auto UVIslandPredicate = [this](int32 Triangle0, int32 Triangle1)
 	{
diff --git a/Engine/Plugins/Experimental/UVEditor/Source/UVEditorTools/Private/UVEditorTransformTool.cpp b/Engine/Plugins/Experimental/UVEditor/Source/UVEditorTools/Private/UVEditorTransformTool.cpp
--- a/Engine/Plugins/Experimental/UVEditor/Source/UVEditorTools/Private/UVEditorTransformTool.cpp
+++ b/Engine/Plugins/Experimental/UVEditor/Source/UVEditorTools/Private/UVEditorTransformTool.cpp
@@ -173,7 +173,7 @@ void UUVEditorTransformTool::Setup()
 		Factory->Settings = Settings;
 		Factory->TransformType = ToolMode.Get(EUVEditorUVTransformType::Transform);
 		Factory->OriginalMesh = Target.UnwrapCanonical;
-		Factory->GetSelectedUVChannel = [&Target]() { return Target.UVLayerIndex; };
+		Factory->GetSelectedUVChannel = [&Target]() { return 0; /*Since we're passing in the unwrap mesh, the UV index is always zero.*/ };
 		if (Selection)
 		{
 			// If we have a selection, it's in the unwrapped mesh. We need both triangles, which are 1:1 between them,
