UVEditor: Fix issue where the Transform's translate mode "absolute" would act as a negative offset for the requested translation values, rather than the starting point as intended.

diff --git a/Engine/Plugins/Experimental/UVEditor/Source/UVEditorTools/Private/Operators/UVEditorUVTransformOp.cpp b/Engine/Plugins/Experimental/UVEditor/Source/UVEditorTools/Private/Operators/UVEditorUVTransformOp.cpp
--- a/Engine/Plugins/Experimental/UVEditor/Source/UVEditorTools/Private/Operators/UVEditorUVTransformOp.cpp
+++ b/Engine/Plugins/Experimental/UVEditor/Source/UVEditorTools/Private/Operators/UVEditorUVTransformOp.cpp
@@ -336,7 +336,7 @@ void FUVEditorUVTransformOp::HandleTransformationOp(FProgressCancel * Progress)
 			RotationPivot = GetPivotFromMode(Vid, PivotMode);
 		}				
 		FVector2f UV = TransformOpLocals::UnwrapPositionToUV(ResultMesh->GetVertexRef(Vid));
-		UV = (UV + FVector2f(Translation) - RotationPivot);
+		UV = (UV + FVector2f(Translation) + RotationPivot);
 		ResultMesh->SetVertex(Vid, TransformOpLocals::UVToUnwrapPosition(UV));
 	};
 
