UVEditor: Clean up options for the various UV Unwrap algorithms to only be visible to the user when the appropriate method has been selected to reduce confusion in the UX.

diff --git a/Engine/Plugins/Experimental/UVEditor/Source/UVEditorTools/Public/Operators/UVEditorRecomputeUVsOp.h b/Engine/Plugins/Experimental/UVEditor/Source/UVEditorTools/Public/Operators/UVEditorRecomputeUVsOp.h
--- a/Engine/Plugins/Experimental/UVEditor/Source/UVEditorTools/Public/Operators/UVEditorRecomputeUVsOp.h
+++ b/Engine/Plugins/Experimental/UVEditor/Source/UVEditorTools/Public/Operators/UVEditorRecomputeUVsOp.h
@@ -96,27 +96,27 @@ public:
 	EUVEditorRecomputeUVsToolOrientationMode AutoRotation = EUVEditorRecomputeUVsToolOrientationMode::MinBounds;
 
 	/**  If enabled, reduces distortion for meshes with triangles of vastly different sizes, This is only enabled if the Unwrap Type is set to Spectral Conformal. */
-	UPROPERTY(EditAnywhere, Category = "UV Unwrap", meta = (EditCondition = "UnwrapType == EUVEditorRecomputeUVsPropertiesUnwrapType::SpectralConformal"))
+	UPROPERTY(EditAnywhere, Category = "UV Unwrap", meta = (EditCondition = "UnwrapType == EUVEditorRecomputeUVsPropertiesUnwrapType::SpectralConformal", EditConditionHides, HideEditConditionToggle))
 	bool bPreserveIrregularity = true;
 
 	/** Number of smoothing steps to apply; this slightly increases distortion but produces more stable results. This is only enabled if the Unwrap Type is set to ExpMap or Island Merging. */
 	UPROPERTY(EditAnywhere, Category = "UV Unwrap", meta = (UIMin = "0", UIMax = "25", ClampMin = "0", ClampMax = "1000",
-		EditCondition = "UnwrapType == EUVEditorRecomputeUVsPropertiesUnwrapType::ExpMap || UnwrapType == EUVEditorRecomputeUVsPropertiesUnwrapType::IslandMerging"))
+		EditCondition = "UnwrapType == EUVEditorRecomputeUVsPropertiesUnwrapType::ExpMap || UnwrapType == EUVEditorRecomputeUVsPropertiesUnwrapType::IslandMerging", EditConditionHides, HideEditConditionToggle))
 	int SmoothingSteps = 5;
 
 	/** Smoothing parameter; larger values result in faster smoothing in each step. This is only enabled if the Unwrap Type is set to ExpMap or Island Merging. */
 	UPROPERTY(EditAnywhere, Category = "UV Unwrap", meta = (UIMin = "0", UIMax = "1.0", ClampMin = "0", ClampMax = "1.0",
-		EditCondition = "UnwrapType == EUVEditorRecomputeUVsPropertiesUnwrapType::ExpMap || UnwrapType == EUVEditorRecomputeUVsPropertiesUnwrapType::IslandMerging"))
+		EditCondition = "UnwrapType == EUVEditorRecomputeUVsPropertiesUnwrapType::ExpMap || UnwrapType == EUVEditorRecomputeUVsPropertiesUnwrapType::IslandMerging", EditConditionHides, HideEditConditionToggle))
 	float SmoothingAlpha = 0.25f;
 
 	/** Threshold for stretching and distortion below which island merging is allowed; larger values increase the allowable UV distortion. This is only enabled if the Unwrap Type is set to Island Merging. */
 	UPROPERTY(EditAnywhere, Category = "UV Unwrap", meta = (DisplayName = "Distortion Threshold", UIMin = "1.0", UIMax = "5.0", ClampMin = "1.0",
-		EditCondition = "UnwrapType == EUVEditorRecomputeUVsPropertiesUnwrapType::IslandMerging"))
+		EditCondition = "UnwrapType == EUVEditorRecomputeUVsPropertiesUnwrapType::IslandMerging", EditConditionHides, HideEditConditionToggle))
 	float MergingDistortionThreshold = 1.5f;
 
 	/** Threshold for the average face normal deviation below  which island merging is allowed. This is only enabled if the Unwrap Type is set to Island Merging. */
 	UPROPERTY(EditAnywhere, Category = "UV Unwrap", meta = (DisplayName = "Angle Threshold", UIMin = "0.0", UIMax = "90.0", ClampMin = "0.0", ClampMax = "180.0",
-		EditCondition = "UnwrapType == EUVEditorRecomputeUVsPropertiesUnwrapType::IslandMerging"))
+		EditCondition = "UnwrapType == EUVEditorRecomputeUVsPropertiesUnwrapType::IslandMerging", EditConditionHides, HideEditConditionToggle))
 	float MergingAngleThreshold = 45.0f;
 
 	/** Uniformly scale and translate UV islands collectively to pack them into the unit square, i.e. fit between 0 and 1 with no overlap. */
@@ -125,12 +125,12 @@ public:
 
 	/** Expected resolution of the output textures; this controls spacing left between UV islands to avoid interpolation artifacts. This is only enabled when the Layout Type is set to Repack. */
 	UPROPERTY(EditAnywhere, Category = "UV Layout", meta = (UIMin = "64", UIMax = "2048", ClampMin = "2", ClampMax = "4096",
-		EditCondition = "LayoutType == EUVEditorRecomputeUVsPropertiesLayoutType::Repack"))
+		EditCondition = "LayoutType == EUVEditorRecomputeUVsPropertiesLayoutType::Repack", EditConditionHides, HideEditConditionToggle))
 	int TextureResolution = 1024;
 
 	/** Scaling factor used for UV island normalization/scaling. This is only enabled when the Layout Type is set to Normalize to Bounds or Normalize to World. */
 	UPROPERTY(EditAnywhere, Category = "UV Layout", meta = (UIMin = "0.001", UIMax = "10", ClampMin = "0.00001", ClampMax = "1000000.0",
-		EditCondition = "LayoutType == EUVEditorRecomputeUVsPropertiesLayoutType::NormalizeToBounds || LayoutType == EUVEditorRecomputeUVsPropertiesLayoutType::NormalizeToWorld"))
+		EditCondition = "LayoutType == EUVEditorRecomputeUVsPropertiesLayoutType::NormalizeToBounds || LayoutType == EUVEditorRecomputeUVsPropertiesLayoutType::NormalizeToWorld", EditConditionHides, HideEditConditionToggle))
 	float NormalizeScale = 1.0f;
 
 	/** Enable UDIM aware layout and keep islands within their originating UDIM tiles when laying out.*/
