UVEditor: Fix crash bug when selecting elements when multiple assets are loaded in the UV Editor.

diff --git a/Engine/Plugins/Experimental/UVEditor/Source/UVEditorTools/Private/Selection/UVToolSelectionHighlightMechanic.cpp b/Engine/Plugins/Experimental/UVEditor/Source/UVEditorTools/Private/Selection/UVToolSelectionHighlightMechanic.cpp
--- a/Engine/Plugins/Experimental/UVEditor/Source/UVEditorTools/Private/Selection/UVToolSelectionHighlightMechanic.cpp
+++ b/Engine/Plugins/Experimental/UVEditor/Source/UVEditorTools/Private/Selection/UVToolSelectionHighlightMechanic.cpp
@@ -203,8 +203,8 @@ void UUVToolSelectionHighlightMechanic::RebuildUnwrapHighlight(
 
 			UBasic2DTriangleSetComponent* UnwrapTriangleSetPtr = UnwrapTriangleSet.Get();
 			UBasic2DLineSetComponent* UnwrapLineSetPtr = UnwrapLineSet.Get();
-			UnwrapTriangleSetPtr->ReserveElements(Selection.SelectedIDs.Num());
-			UnwrapLineSetPtr->ReserveElements(Selection.SelectedIDs.Num() * 3);
+			UnwrapTriangleSetPtr->ReserveAdditionalElements(Selection.SelectedIDs.Num());
+			UnwrapLineSetPtr->ReserveAdditionalElements(Selection.SelectedIDs.Num() * 3);
 			for (int32 Tid : Selection.SelectedIDs)
 			{
 				if (!ensure(Mesh.IsTriangle(Tid)))
@@ -247,10 +247,10 @@ void UUVToolSelectionHighlightMechanic::RebuildUnwrapHighlight(
 			UBasic2DLineSetComponent* SewEdgePairingRightLineSetPtr = SewEdgePairingRightLineSet.Get();
 			UBasic2DLineSetComponent* SewEdgeUnselectedPairingLineSetPtr = SewEdgeUnselectedPairingLineSet.Get();
 
-			UnwrapLineSetPtr->ReserveElements(Selection.SelectedIDs.Num());
-			SewEdgePairingLeftLineSetPtr->ReserveElements(Selection.SelectedIDs.Num());
-			SewEdgePairingRightLineSetPtr->ReserveElements(Selection.SelectedIDs.Num());
-			SewEdgeUnselectedPairingLineSetPtr->ReserveElements(Selection.SelectedIDs.Num());
+			UnwrapLineSetPtr->ReserveAdditionalElements(Selection.SelectedIDs.Num());
+			SewEdgePairingLeftLineSetPtr->ReserveAdditionalElements(Selection.SelectedIDs.Num());
+			SewEdgePairingRightLineSetPtr->ReserveAdditionalElements(Selection.SelectedIDs.Num());
+			SewEdgeUnselectedPairingLineSetPtr->ReserveAdditionalElements(Selection.SelectedIDs.Num());
 			for (int32 Eid : Selection.SelectedIDs)
 			{
 				if (!ensure(Mesh.IsEdge(Eid)))
@@ -316,7 +316,7 @@ void UUVToolSelectionHighlightMechanic::RebuildUnwrapHighlight(
 			TRACE_CPUPROFILER_EVENT_SCOPE(MeshSelectionMechanic_RebuildDrawnElements_Vertex);
 
 			UBasic2DPointSetComponent* UnwrapPointSetPtr = UnwrapPointSet.Get();
-			UnwrapPointSetPtr->ReserveElements(Selection.SelectedIDs.Num());
+			UnwrapPointSetPtr->ReserveAdditionalElements(Selection.SelectedIDs.Num());
 
 			for (const int32 Vid : Selection.SelectedIDs)
 			{
@@ -435,7 +435,7 @@ void UUVToolSelectionHighlightMechanic::RebuildAppliedHighlightFromUnwrapSelecti
 		{
 			TRACE_CPUPROFILER_EVENT_SCOPE(Triangle);
 
-			LivePreviewLineSet->ReserveElements(Selection.SelectedIDs.Num()*3);
+			LivePreviewLineSetPtr->ReserveAdditionalElements(Selection.SelectedIDs.Num()*3);
 			for (int32 Tid : Selection.SelectedIDs)
 			{
 				if (!ensure(AppliedMesh.IsTriangle(Tid)))
@@ -463,7 +463,7 @@ void UUVToolSelectionHighlightMechanic::RebuildAppliedHighlightFromUnwrapSelecti
 		else if (Selection.Type == FUVToolSelection::EType::Edge)
 		{
 			TRACE_CPUPROFILER_EVENT_SCOPE(Edge);
-			LivePreviewLineSet->ReserveElements(Selection.SelectedIDs.Num());
+			LivePreviewLineSetPtr->ReserveAdditionalElements(Selection.SelectedIDs.Num());
 
 			for (int32 UnwrapEid : Selection.SelectedIDs)
 			{
@@ -487,7 +487,7 @@ void UUVToolSelectionHighlightMechanic::RebuildAppliedHighlightFromUnwrapSelecti
 		{
 			TRACE_CPUPROFILER_EVENT_SCOPE(Vertex);
 
-			LivePreviewPointSetPtr->ReserveElements(Selection.SelectedIDs.Num());
+			LivePreviewPointSetPtr->ReserveAdditionalElements(Selection.SelectedIDs.Num());
 
 			for (const int32 UnwrapVid : Selection.SelectedIDs)
 			{
@@ -537,6 +537,7 @@ void UUVToolSelectionHighlightMechanic::AppendAppliedHighlight(const TArray<FUVT
 		{
 			TRACE_CPUPROFILER_EVENT_SCOPE(Triangle);
 
+			LivePreviewLineSetPtr->ReserveAdditionalElements(Selection.SelectedIDs.Num() * 3);
 			for (int32 Tid : Selection.SelectedIDs)
 			{
 				if (!ensure(AppliedMesh.IsTriangle(Tid)))
@@ -565,6 +566,7 @@ void UUVToolSelectionHighlightMechanic::AppendAppliedHighlight(const TArray<FUVT
 		{
 			TRACE_CPUPROFILER_EVENT_SCOPE(Edge);
 
+			LivePreviewLineSetPtr->ReserveAdditionalElements(Selection.SelectedIDs.Num());
 			for (int32 Eid : Selection.SelectedIDs)
 			{
 				if (!ensure(AppliedMesh.IsEdge(Eid)))
@@ -580,7 +582,7 @@ void UUVToolSelectionHighlightMechanic::AppendAppliedHighlight(const TArray<FUVT
 		{
 			TRACE_CPUPROFILER_EVENT_SCOPE(Vertex);
 
-			LivePreviewPointSetPtr->ReserveElements(Selection.SelectedIDs.Num());
+			LivePreviewPointSetPtr->ReserveAdditionalElements(Selection.SelectedIDs.Num());
 
 			for (const int32 Vid : Selection.SelectedIDs)
 			{
diff --git a/Engine/Plugins/Experimental/UVEditor/Source/UVEditorTools/Public/Drawing/BasicElementSet.h b/Engine/Plugins/Experimental/UVEditor/Source/UVEditorTools/Public/Drawing/BasicElementSet.h
--- a/Engine/Plugins/Experimental/UVEditor/Source/UVEditorTools/Public/Drawing/BasicElementSet.h
+++ b/Engine/Plugins/Experimental/UVEditor/Source/UVEditorTools/Public/Drawing/BasicElementSet.h
@@ -44,11 +44,17 @@ public:
 		return Num;
 	}
 
-	/** Reserve enough memory for up to the given maximum number of points. This function needs to be called before any calls to AddPoint(). */
+	/** Reserve enough memory for up to the given maximum number of points. This function needs to be called before any calls to AddElement(). */
 	void ReserveElements(const int32 MaxNum)
 	{
 		Elements.SetNumUninitialized(MaxNum * ElementDataSize);
 	}
+	
+	/** Reserve enough memory for an additional number of points over the existing number already reserved. This function needs to be called before any calls to AddElement(). */
+	void ReserveAdditionalElements(const int32 AdditionalNum)
+	{
+		Elements.SetNumUninitialized(Elements.Num() + AdditionalNum * ElementDataSize);
+	}
 
 	/** Add a point to be rendered using the component. */
 	template<typename... Targs>
