UVEditor: Fix various UV editor undo/redo transactions not being expired when the editor is closed.

diff --git a/Engine/Plugins/Experimental/UVEditor/Source/UVEditor/Private/UVEditorMode.cpp b/Engine/Plugins/Experimental/UVEditor/Source/UVEditor/Private/UVEditorMode.cpp
--- a/Engine/Plugins/Experimental/UVEditor/Source/UVEditor/Private/UVEditorMode.cpp
+++ b/Engine/Plugins/Experimental/UVEditor/Source/UVEditor/Private/UVEditorMode.cpp
@@ -107,11 +107,9 @@ namespace UVEditorModeLocals
 
 		virtual bool HasExpired(UObject* Object) const override
 		{
-			// To not be expired, we must be in some tool.
+			// To not be expired, we must be active and in some non-default tool.
 			UUVEditorMode* Mode = Cast<UUVEditorMode>(Object);
-			return !(Mode && Mode->GetInteractiveToolsContext()
-				&& Mode->GetInteractiveToolsContext()->ToolManager
-				&& Mode->GetInteractiveToolsContext()->ToolManager->HasAnyActiveTool());
+			return !(Mode && Mode->IsActive() && !Mode->IsDefaultToolActive());
 		}
 
 		virtual FString ToString() const override
@@ -667,7 +665,7 @@ void UUVEditorMode::ActivateDefaultTool()
 
 bool UUVEditorMode::IsDefaultToolActive()
 {
-	return GetInteractiveToolsContext()->IsToolActive(EToolSide::Mouse, DefaultToolIdentifier);
+	return GetInteractiveToolsContext() && GetInteractiveToolsContext()->IsToolActive(EToolSide::Mouse, DefaultToolIdentifier);
 }
 
 void UUVEditorMode::BindCommands()
diff --git a/Engine/Plugins/Experimental/UVEditor/Source/UVEditorTools/Private/Selection/UVToolSelectionAPI.cpp b/Engine/Plugins/Experimental/UVEditor/Source/UVEditorTools/Private/Selection/UVToolSelectionAPI.cpp
--- a/Engine/Plugins/Experimental/UVEditor/Source/UVEditorTools/Private/Selection/UVToolSelectionAPI.cpp
+++ b/Engine/Plugins/Experimental/UVEditor/Source/UVEditorTools/Private/Selection/UVToolSelectionAPI.cpp
@@ -59,6 +59,8 @@ void UUVToolSelectionAPI::Initialize(
 	SelectionMechanic->Initialize(UnwrapWorld, LivePreviewAPI->GetLivePreviewWorld(), this);
 
 	UnwrapInputRouter->RegisterSource(MechanicAdapter);
+
+	bIsActive = true;
 }
 
 void UUVToolSelectionAPI::SetTargets(const TArray<TObjectPtr<UUVEditorToolMeshInput>>& TargetsIn)
@@ -105,6 +107,8 @@ void UUVToolSelectionAPI::Shutdown()
 		Target->OnCanonicalModified.RemoveAll(this);
 	}
 	Targets.Empty();
+
+	bIsActive = false;
 }
 
 void UUVToolSelectionAPI::ClearSelections(bool bBroadcast, bool bEmitChange)
@@ -595,6 +599,12 @@ TArray<FUVToolSelection> UUVToolSelectionAPI::FSelectionChange::GetBefore() cons
 	return SelectionsOut;
 }
 
+bool UUVToolSelectionAPI::FSelectionChange::HasExpired(UObject* Object) const
+{
+	UUVToolSelectionAPI* SelectionAPI = Cast<UUVToolSelectionAPI>(Object);
+	return !(SelectionAPI && SelectionAPI->IsActive());
+}
+
 void UUVToolSelectionAPI::FSelectionChange::Apply(UObject* Object) 
 {
 	UUVToolSelectionAPI* SelectionAPI = Cast<UUVToolSelectionAPI>(Object);
@@ -657,6 +667,12 @@ const TArray<FUVToolSelection>& UUVToolSelectionAPI::FUnsetSelectionChange::GetB
 	return Before;
 }
 
+bool UUVToolSelectionAPI::FUnsetSelectionChange::HasExpired(UObject* Object) const
+{
+	UUVToolSelectionAPI* SelectionAPI = Cast<UUVToolSelectionAPI>(Object);
+	return !(SelectionAPI && SelectionAPI->IsActive());
+}
+
 void UUVToolSelectionAPI::FUnsetSelectionChange::Apply(UObject* Object)
 {
 	UUVToolSelectionAPI* SelectionAPI = Cast<UUVToolSelectionAPI>(Object);
diff --git a/Engine/Plugins/Experimental/UVEditor/Source/UVEditorTools/Public/Selection/UVToolSelectionAPI.h b/Engine/Plugins/Experimental/UVEditor/Source/UVEditorTools/Public/Selection/UVToolSelectionAPI.h
--- a/Engine/Plugins/Experimental/UVEditor/Source/UVEditorTools/Public/Selection/UVToolSelectionAPI.h
+++ b/Engine/Plugins/Experimental/UVEditor/Source/UVEditorTools/Public/Selection/UVToolSelectionAPI.h
@@ -271,6 +271,7 @@ public:
 		// emitted. Respects bUseStableUnwrapCanonicalIDsForEdges.
 		virtual TArray<FUVToolSelection> GetBefore() const;
 
+		virtual bool HasExpired(UObject* Object) const override;
 		virtual void Apply(UObject* Object) override;
 		virtual void Revert(UObject* Object) override;
 		virtual FString ToString() const override;
@@ -292,6 +293,7 @@ public:
 
 		virtual const TArray<FUVToolSelection>& GetBefore() const;
 
+		virtual bool HasExpired(UObject* Object) const override;
 		virtual void Apply(UObject* Object) override;
 		virtual void Revert(UObject* Object) override;
 		virtual FString ToString() const override;
@@ -301,6 +303,10 @@ public:
 		TArray<FUVToolSelection> After;
 	};
 
+	// Returns true in between Initialize() and Shutdown(). Used by undo/redo to expire transactions
+	// when editor is closed.
+	bool IsActive() const { return bIsActive; }
+
 	/**
 	 * Preps a selection change transaction, if the user wants more control on what the previous and
 	 * current selection is.
@@ -359,6 +365,8 @@ protected:
 
 	FVector3d CachedUnwrapSelectionBoundingBoxCenter;
 	bool bCachedUnwrapSelectionBoundingBoxCenterValid = false;
+
+	bool bIsActive = false;
 };
 
 // UInterface for IUVToolSupportsSelection
