UVEditor: Change the property set for the Transform/Align/Distribute tools to be transacting, allowing for capturing of button based actions done through their respective details customizations in the undo stack.

diff --git a/Engine/Plugins/Experimental/UVEditor/Source/UVEditorTools/Public/Operators/UVEditorUVTransformOp.h b/Engine/Plugins/Experimental/UVEditor/Source/UVEditorTools/Public/Operators/UVEditorUVTransformOp.h
--- a/Engine/Plugins/Experimental/UVEditor/Source/UVEditorTools/Public/Operators/UVEditorUVTransformOp.h
+++ b/Engine/Plugins/Experimental/UVEditor/Source/UVEditorTools/Public/Operators/UVEditorUVTransformOp.h
@@ -115,6 +115,17 @@ class UVEDITORTOOLS_API UUVEditorUVTransformPropertiesBase : public UInteractive
 {
 	GENERATED_BODY()
 public:
+
+	UUVEditorUVTransformPropertiesBase()
+	{
+		// TODO: This makes the changes to the Transform, Align, and Distribute property sets show up
+		// in the Undo history, so the button presses in Quick Translate, Align and Distribute modes
+		// can be undone. However, it doens't properly track tool lifetime and therefore doesn't handle
+		// correct expiry when the tools using the properties shutdown. A custom change should probably
+		// be implemented to handle this behavior more correctly.
+		SetFlags(RF_Transactional);
+	}
+
 };
 
 
