UVEditor: Fix minor issues with Distribute tool to make per vertex distribution visually more coherent between modes. Also adjusts horizontal and vertical center distribution to be more aligned with the positioning of the other modes.

diff --git a/Engine/Plugins/Experimental/UVEditor/Source/UVEditorTools/Private/Operators/UVEditorUVTransformOp.cpp b/Engine/Plugins/Experimental/UVEditor/Source/UVEditorTools/Private/Operators/UVEditorUVTransformOp.cpp
--- a/Engine/Plugins/Experimental/UVEditor/Source/UVEditorTools/Private/Operators/UVEditorUVTransformOp.cpp
+++ b/Engine/Plugins/Experimental/UVEditor/Source/UVEditorTools/Private/Operators/UVEditorUVTransformOp.cpp
@@ -36,7 +36,7 @@ void FUVEditorUVTransformBaseOp::SetSelection(const TSet<int32>& TriangleSelecti
 		VertexSelection = VertexSelectionIn;
 }
 
-void FUVEditorUVTransformBaseOp::SortComponentsByBoundingBox()
+void FUVEditorUVTransformBaseOp::SortComponentsByBoundingBox(bool bSortXThenY)
 {
 	SpatiallyOrderedComponentIndex.SetNum(UVComponents->Num());
 	for (int32 Index = 0; Index < UVComponents->Num(); ++Index)
@@ -46,15 +46,17 @@ void FUVEditorUVTransformBaseOp::SortComponentsByBoundingBox()
 	TArrayView<int32> IndexView(SpatiallyOrderedComponentIndex.GetData(), SpatiallyOrderedComponentIndex.Num());
 	TArrayView<FAxisAlignedBox2d> BBView(PerComponentBoundingBoxes.GetData(), PerComponentBoundingBoxes.Num());
 
-	auto BBIndexSort = [&BBView](const int32& A, const int32& B)
+	const int32 SortDirection = bSortXThenY ? 1 : 0;
+
+	auto BBIndexSort = [&BBView, SortDirection](const int32& A, const int32& B)
 	{
 		FVector2D CenterDiff = BBView[A].Center() - BBView[B].Center();
 
-		if (CenterDiff.X < 0)
+		if (CenterDiff[1-SortDirection] < 0)
 		{
 			return true;
 		}
-		else if (FMath::IsNearlyZero(CenterDiff.X) && CenterDiff.Y < 0)
+		else if (FMath::IsNearlyZero(CenterDiff[1 - SortDirection]) && CenterDiff[SortDirection] < 0)
 		{
 			return true;
 		}
@@ -522,13 +524,13 @@ void FUVEditorUVAlignOp::HandleTransformationOp(FProgressCancel* Progress)
 void FUVEditorUVDistributeOp::HandleTransformationOp(FProgressCancel* Progress)
 {
 	int32 NumComponents = UVComponents->Num();
-	RebuildBoundingBoxes();
-	SortComponentsByBoundingBox();
+	RebuildBoundingBoxes();	
 	TArray<FVector2f> PerComponentTranslation;
 	
 
 	auto ComputeDistributeTranslations = [this, NumComponents](bool bVertical, EUVEditorAlignDirectionBackend EdgeDirection, float SpreadDirection, bool bEqualizeSpacing)
 	{
+		SortComponentsByBoundingBox(!bVertical);
 		TArray<FVector2f> PerComponentTranslation;
 		PerComponentTranslation.SetNum(NumComponents);
 		float TotalDistance = 0.0f;
@@ -549,7 +551,19 @@ void FUVEditorUVDistributeOp::HandleTransformationOp(FProgressCancel* Progress)
 		float PerComponentSpace = BoundingBoxDistance / (NumComponents - 1);
 
 		float NextPosition = 0.0f;
-		FVector2f OverallAlignmentPoint = GetAlignmentPointFromBoundingBoxAndDirection(EdgeDirection, OverallBoundingBox);
+		FVector2f OverallAlignmentPoint;
+		switch (EdgeDirection) {
+		case EUVEditorAlignDirectionBackend::CenterHorizontally:
+			OverallAlignmentPoint = GetAlignmentPointFromBoundingBoxAndDirection(EUVEditorAlignDirectionBackend::Left, OverallBoundingBox);
+			break;
+		case EUVEditorAlignDirectionBackend::CenterVertically:
+			OverallAlignmentPoint = GetAlignmentPointFromBoundingBoxAndDirection(EUVEditorAlignDirectionBackend::Bottom, OverallBoundingBox);
+			break;
+		default:
+			OverallAlignmentPoint = GetAlignmentPointFromBoundingBoxAndDirection(EdgeDirection, OverallBoundingBox);
+			break;
+		}
+
 		for (int32 Cid: SpatiallyOrderedComponentIndex)
 		{
 			FVector2f ComponentAlignmentPoint = GetAlignmentPointFromBoundingBoxAndDirection(EdgeDirection, PerComponentBoundingBoxes[Cid]);
@@ -576,6 +590,7 @@ void FUVEditorUVDistributeOp::HandleTransformationOp(FProgressCancel* Progress)
 
 	auto ComputeMinimallyRemoveOverlap = [this, NumComponents]()
 	{
+		SortComponentsByBoundingBox();
 		FVector2D Center = OverallBoundingBox.Center();
 
 		// Create a local copy so we can move things around...
diff --git a/Engine/Plugins/Experimental/UVEditor/Source/UVEditorTools/Public/Operators/UVEditorUVTransformOp.h b/Engine/Plugins/Experimental/UVEditor/Source/UVEditorTools/Public/Operators/UVEditorUVTransformOp.h
--- a/Engine/Plugins/Experimental/UVEditor/Source/UVEditorTools/Public/Operators/UVEditorUVTransformOp.h
+++ b/Engine/Plugins/Experimental/UVEditor/Source/UVEditorTools/Public/Operators/UVEditorUVTransformOp.h
@@ -347,7 +347,7 @@ protected:
 	FVector2f GetAlignmentPointFromUDIMAndDirection(EUVEditorAlignDirectionBackend Direction, FVector2i UDIMTile);
 	void RebuildBoundingBoxes();
 	void CollectTransformElements();
-	void SortComponentsByBoundingBox();
+	void SortComponentsByBoundingBox(bool bSortXThenY = true);
 
 	TMap<int32, int32> ElementToComponent;
 	FDynamicMeshUVOverlay* ActiveUVLayer;
