UVEditor: Prevent NoAction item from appearing in the Action dropdown for the channel edit tool.

diff --git a/Engine/Plugins/Experimental/UVEditor/Source/UVEditorTools/Public/UVEditorLayerEditTool.h b/Engine/Plugins/Experimental/UVEditor/Source/UVEditorTools/Public/UVEditorLayerEditTool.h
--- a/Engine/Plugins/Experimental/UVEditor/Source/UVEditorTools/Public/UVEditorLayerEditTool.h
+++ b/Engine/Plugins/Experimental/UVEditor/Source/UVEditorTools/Public/UVEditorLayerEditTool.h
@@ -45,7 +45,7 @@ class UVEDITORTOOLS_API UUVEditorChannelEditSettings : public UInteractiveToolPr
 	GENERATED_BODY()
 public:
 	/** UV Layer Edit action to preform */
-	UPROPERTY(EditAnywhere, Category = Options)
+	UPROPERTY(EditAnywhere, Category = Options, meta = (InvalidEnumValues = "NoAction"))
 	EChannelEditToolAction Action = EChannelEditToolAction::Add;
 };
 
