UVEditor: Added support for a pivot visualization to avoid UX confusion during transform and align tool operations.

diff --git a/Engine/Plugins/Experimental/UVEditor/Source/UVEditor/Private/UVEditorMode.cpp b/Engine/Plugins/Experimental/UVEditor/Source/UVEditor/Private/UVEditorMode.cpp
--- a/Engine/Plugins/Experimental/UVEditor/Source/UVEditor/Private/UVEditorMode.cpp
+++ b/Engine/Plugins/Experimental/UVEditor/Source/UVEditor/Private/UVEditorMode.cpp
@@ -639,6 +639,19 @@ UObject* UUVEditorMode::GetUDIMSettingsObject()
 	return nullptr;
 }
 
+UObject* UUVEditorMode::GetToolDisplaySettingsObject()
+{
+	UContextObjectStore* ContextStore = GetInteractiveToolsContext()->ToolManager->GetContextObjectStore();
+	UUVEditorToolPropertiesAPI* ToolPropertiesAPI = ContextStore->FindContext<UUVEditorToolPropertiesAPI>();
+
+	if (ToolPropertiesAPI)
+	{
+		return ToolPropertiesAPI->GetToolDisplayProperties();
+	}
+	return nullptr;
+}
+
+
 void UUVEditorMode::Render(IToolsContextRenderAPI* RenderAPI)
 {
 	if (SelectionAPI)
@@ -961,6 +974,8 @@ void UUVEditorMode::InitializeModeContexts()
 		GetInteractiveToolsContext()->InputRouter, LivePreviewAPI, EmitChangeAPI);
 	AddContextObject(SelectionAPI);
 
+	UUVEditorToolPropertiesAPI* UVEditorToolPropertiesAPI = NewObject<UUVEditorToolPropertiesAPI>();
+	AddContextObject(UVEditorToolPropertiesAPI);
 }
 
 void UUVEditorMode::InitializeTargets()
diff --git a/Engine/Plugins/Experimental/UVEditor/Source/UVEditor/Private/UVEditorModeToolkit.cpp b/Engine/Plugins/Experimental/UVEditor/Source/UVEditor/Private/UVEditorModeToolkit.cpp
--- a/Engine/Plugins/Experimental/UVEditor/Source/UVEditor/Private/UVEditorModeToolkit.cpp
+++ b/Engine/Plugins/Experimental/UVEditor/Source/UVEditor/Private/UVEditorModeToolkit.cpp
@@ -271,6 +271,22 @@ TSharedRef<SWidget> FUVEditorModeToolkit::CreateUDIMSettingsWidget()
 	return CreateDisplaySettingsWidget(Mode->GetUDIMSettingsObject());
 }
 
+TSharedRef<SWidget> FUVEditorModeToolkit::GetToolDisplaySettingsWidget()
+{
+	UUVEditorMode* Mode = Cast<UUVEditorMode>(GetScriptableEditorMode());
+	UObject* SettingsObject = Mode->GetToolDisplaySettingsObject();
+	if (SettingsObject)
+	{
+		return CreateDisplaySettingsWidget(SettingsObject);
+	}
+	else
+	{
+		return SNew(SBorder)
+			.BorderImage(FAppStyle::GetBrush("NoBorder"))
+			.Padding(0);
+	}
+}
+
 TSharedRef<SWidget> FUVEditorModeToolkit::CreateDisplaySettingsWidget(UObject* SettingsObject) const
 {
 	TSharedRef<SBorder> GridDetailsContainer =
diff --git a/Engine/Plugins/Experimental/UVEditor/Source/UVEditor/Private/UVEditorToolkit.cpp b/Engine/Plugins/Experimental/UVEditor/Source/UVEditor/Private/UVEditorToolkit.cpp
--- a/Engine/Plugins/Experimental/UVEditor/Source/UVEditor/Private/UVEditorToolkit.cpp
+++ b/Engine/Plugins/Experimental/UVEditor/Source/UVEditor/Private/UVEditorToolkit.cpp
@@ -502,6 +502,17 @@ void FUVEditorToolkit::PostInitAssetEditor()
 				]
 				];
 
+			Container->AddSlot()
+				.AutoHeight()
+				.Padding(FMargin(0.f, 0.f, 8.f, 0.f))
+				[
+					SNew(SBox)
+					.MinDesiredWidth(500)
+				[
+					UVModeToolkit->GetToolDisplaySettingsWidget()
+				]
+				];
+
 			bool bEnableUDIMSupport = (FUVEditorUXSettings::CVarEnablePrototypeUDIMSupport.GetValueOnGameThread() > 0);
 			if (bEnableUDIMSupport)
 			{
diff --git a/Engine/Plugins/Experimental/UVEditor/Source/UVEditor/Public/UVEditorMode.h b/Engine/Plugins/Experimental/UVEditor/Source/UVEditor/Public/UVEditorMode.h
--- a/Engine/Plugins/Experimental/UVEditor/Source/UVEditor/Public/UVEditorMode.h
+++ b/Engine/Plugins/Experimental/UVEditor/Source/UVEditor/Public/UVEditorMode.h
@@ -42,11 +42,11 @@ class UVEDITOR_API UUVEditorGridProperties : public UInteractiveToolPropertySet
 	GENERATED_BODY()
 public:
 	/** Should the grid be shown?*/
-	UPROPERTY(EditAnywhere, Category = Grid, meta = (DisplayName = "Display Grid"))
+	UPROPERTY(EditAnywhere, Category = "Grid & Guides", meta = (DisplayName = "Display Grid"))
 	bool bDrawGrid = true;
 
 	/** Should the grid rulers be shown?*/
-	UPROPERTY(EditAnywhere, Category = Grid, meta = (DisplayName = "Display Rulers"))
+	UPROPERTY(EditAnywhere, Category = "Grid & Guides", meta = (DisplayName = "Display Rulers"))
 	bool bDrawRulers = true;
 };
 
@@ -218,6 +218,9 @@ public:
 	/** @return A settings object suitable for display in a details panel to control UDIM configuration. */
 	UObject* GetUDIMSettingsObject();
 
+	/** @return A settings object suitable for display in a details panel to control the active tool's visualization settings. */
+	UObject* GetToolDisplaySettingsObject();
+
 	virtual void Render(IToolsContextRenderAPI* RenderAPI) ;
 	virtual void DrawHUD(FCanvas* Canvas, IToolsContextRenderAPI* RenderAPI);
 
diff --git a/Engine/Plugins/Experimental/UVEditor/Source/UVEditor/Public/UVEditorModeToolkit.h b/Engine/Plugins/Experimental/UVEditor/Source/UVEditor/Public/UVEditorModeToolkit.h
--- a/Engine/Plugins/Experimental/UVEditor/Source/UVEditor/Public/UVEditorModeToolkit.h
+++ b/Engine/Plugins/Experimental/UVEditor/Source/UVEditor/Public/UVEditorModeToolkit.h
@@ -35,6 +35,9 @@ public:
 	/** Creates a widget where the visible UDIMs can be changed. */
 	virtual TSharedRef<SWidget> CreateUDIMSettingsWidget();
 
+	/** Creates a widget where the active tool's display settings can be changed. */
+	virtual TSharedRef<SWidget> GetToolDisplaySettingsWidget();
+
 	/** Enables or diables the Play in Editor mode warning for the editor */
 	virtual void EnableShowPIEWarning(bool bEnable);
 
diff --git a/Engine/Plugins/Experimental/UVEditor/Source/UVEditorTools/Private/Operators/UVEditorUVTransformOp.cpp b/Engine/Plugins/Experimental/UVEditor/Source/UVEditorTools/Private/Operators/UVEditorUVTransformOp.cpp
--- a/Engine/Plugins/Experimental/UVEditor/Source/UVEditorTools/Private/Operators/UVEditorUVTransformOp.cpp
+++ b/Engine/Plugins/Experimental/UVEditor/Source/UVEditorTools/Private/Operators/UVEditorUVTransformOp.cpp
@@ -202,6 +202,7 @@ void FUVEditorUVTransformBaseOp::CalculateResult(FProgressCancel* Progress)
 		return;
 	}
 
+	VisualizationPivots.Empty();
 	int UVLayerInput = UVLayerIndex;
 	ActiveUVLayer = ResultMesh->Attributes()->GetUVLayer(UVLayerInput);
 	ensure(ActiveUVLayer);
@@ -395,6 +396,30 @@ void FUVEditorUVTransformOp::HandleTransformationOp(FProgressCancel * Progress)
 		ApplyTransformFunc(TranslateFunc);
 	}
 
+	// Provide pivots for tool visualization after operation completes.
+	switch (PivotMode)
+	{
+		case EUVEditorPivotTypeBackend::Origin:
+			VisualizationPivots.Add(FVector2D(0, 0));
+			break;
+		case EUVEditorPivotTypeBackend::BoundingBoxCenter:		
+			RebuildBoundingBoxes();
+			VisualizationPivots.Add((FVector2D)OverallBoundingBox.Center());
+			break;
+		case EUVEditorPivotTypeBackend::IndividualBoundingBoxCenter:
+			RebuildBoundingBoxes();
+			for (const FAxisAlignedBox2d& ComponentBoundingBox : PerComponentBoundingBoxes)
+			{
+				VisualizationPivots.Add((FVector2D)ComponentBoundingBox.Center());
+			}
+			break;
+		case EUVEditorPivotTypeBackend::Manual:
+			VisualizationPivots.Add(ManualPivot);
+			break;
+		default:
+			ensure(false);
+	}
+	
 }
 
 
@@ -517,7 +542,26 @@ void FUVEditorUVAlignOp::HandleTransformationOp(FProgressCancel* Progress)
 		}
 	}
 	
-
+	// Provide pivots for tool visualization after operation completes.
+	switch (AlignAnchor)
+	{
+		case EUVEditorAlignAnchorBackend::BoundingBox:
+		{
+			if (AlignDirection != EUVEditorAlignDirectionBackend::None)
+			{
+				FVector2f BoundingBoxAlignmentPoint = GetAlignmentPointFromBoundingBoxAndDirection(AlignDirection, OverallBoundingBox);
+				VisualizationPivots.Add((FVector2D)BoundingBoxAlignmentPoint);
+			}
+		}
+		break;
+		case EUVEditorAlignAnchorBackend::Manual:
+		{
+			VisualizationPivots.Add(ManualAnchor);
+		}
+		break;
+		default:
+			break; // Currently we don't support visualizing anchors for UDIMTiles...
+	}
 }
 
 
diff --git a/Engine/Plugins/Experimental/UVEditor/Source/UVEditorTools/Private/UVEditorTransformTool.cpp b/Engine/Plugins/Experimental/UVEditor/Source/UVEditorTools/Private/UVEditorTransformTool.cpp
--- a/Engine/Plugins/Experimental/UVEditor/Source/UVEditorTools/Private/UVEditorTransformTool.cpp
+++ b/Engine/Plugins/Experimental/UVEditor/Source/UVEditorTools/Private/UVEditorTransformTool.cpp
@@ -138,6 +138,7 @@ void UUVEditorTransformTool::Setup()
 	ToolStartTimeAnalytics = FDateTime::UtcNow();
 
 	UInteractiveTool::Setup();
+	UContextObjectStore* ContextStore = GetToolManager()->GetContextObjectStore();
 
 	switch(ToolMode.Get(EUVEditorUVTransformType::Transform))
 	{
@@ -157,7 +158,18 @@ void UUVEditorTransformTool::Setup()
 	//Settings->bUDIMCVAREnabled = (FUVEditorUXSettings::CVarEnablePrototypeUDIMSupport.GetValueOnGameThread() > 0);
 	AddToolPropertySource(Settings);
 
-	UContextObjectStore* ContextStore = GetToolManager()->GetContextObjectStore();
+	DisplaySettings = NewObject<UUVEditorTransformToolDisplayProperties>(this);
+	DisplaySettings->RestoreProperties(this);
+	DisplaySettings->GetOnModified().AddLambda([this](UObject* PropertySetArg, FProperty* PropertyArg)
+	{
+		OnPropertyModified(PropertySetArg, PropertyArg);
+	});
+	UUVEditorToolPropertiesAPI* PropertiesAPI = ContextStore->FindContext<UUVEditorToolPropertiesAPI>();
+	if (PropertiesAPI)
+	{
+		PropertiesAPI->SetToolDisplayProperties(DisplaySettings);
+	}
+
 	UVToolSelectionAPI = ContextStore->FindContext<UUVToolSelectionAPI>();
 
 	UUVToolSelectionAPI::FHighlightOptions HighlightOptions;
@@ -165,9 +177,13 @@ void UUVEditorTransformTool::Setup()
 	HighlightOptions.bAutoUpdateUnwrap = true;
 	UVToolSelectionAPI->SetHighlightOptions(HighlightOptions);
 	UVToolSelectionAPI->SetHighlightVisible(true, false, true);
-	
+
+	PerTargetPivotLocations.SetNum(Targets.Num());
+
 	auto SetupOpFactory = [this](UUVEditorToolMeshInput& Target, const FUVToolSelection* Selection)
 	{
+		int32 TargetIndex = Targets.Find(&Target);
+
 		TObjectPtr<UUVEditorUVTransformOperatorFactory> Factory = NewObject<UUVEditorUVTransformOperatorFactory>();
 		Factory->TargetTransform = Target.AppliedPreview->PreviewMesh->GetTransform();
 		Factory->Settings = Settings;
@@ -200,6 +216,14 @@ void UUVEditorTransformTool::Setup()
 			this->UVToolSelectionAPI->RebuildUnwrapHighlight(Preview->PreviewMesh->GetTransform());
 			});
 
+		Target.UnwrapPreview->OnOpCompleted.AddLambda(
+			[this, TargetIndex](const FDynamicMeshOperator* Op)
+			{
+				const FUVEditorUVTransformBaseOp* TransformBaseOp = (const FUVEditorUVTransformBaseOp*)(Op);
+				PerTargetPivotLocations[TargetIndex] = TransformBaseOp->GetPivotLocations();
+			}
+		);
+
 		Target.UnwrapPreview->InvalidateResult();
 		return Factory;
 	};
@@ -290,6 +314,7 @@ void UUVEditorTransformTool::Shutdown(EToolShutdownType ShutdownType)
 	}
 
 	Settings = nullptr;
+	DisplaySettings = nullptr;
 	Targets.Empty();
 }
 
@@ -303,6 +328,11 @@ void UUVEditorTransformTool::OnTick(float DeltaTime)
 
 void UUVEditorTransformTool::OnPropertyModified(UObject* PropertySet, FProperty* Property)
 {
+	if (PropertySet == DisplaySettings)
+	{
+		return;
+	}
+
 	for (TObjectPtr<UUVEditorToolMeshInput> Target : Targets)
 	{
 		Target->UnwrapPreview->InvalidateResult();
@@ -323,7 +353,47 @@ bool UUVEditorTransformTool::CanAccept() const
 
 void UUVEditorTransformTool::DrawHUD(FCanvas* Canvas, IToolsContextRenderAPI* RenderAPI)
 {
-	// TODO: Add support here for highlighting first selected item for alignment visualization
+	// TODO: Add support here for highlighting first selected item for alignment visualization	
+
+	auto ConvertUVToPixel = [RenderAPI](const FVector2D& UVIn, FVector2D& PixelOut)
+	{
+		FVector WorldPoint = FUVEditorUXSettings::UVToVertPosition(FUVEditorUXSettings::ExternalUVToInternalUV((FVector2f)UVIn));
+		FVector4 TestProjectedHomogenous = RenderAPI->GetSceneView()->WorldToScreen(WorldPoint);
+		bool bValid = RenderAPI->GetSceneView()->ScreenToPixel(TestProjectedHomogenous, PixelOut);
+		return bValid;
+	};
+
+	if (DisplaySettings->bDrawPivots)
+	{
+		for (int32 TargetIndex = 0; TargetIndex < Targets.Num(); ++TargetIndex)
+		{
+			for (const FVector2D& PivotLocation : PerTargetPivotLocations[TargetIndex])
+			{
+				FVector2D PivotLocationPixel;
+				ConvertUVToPixel(PivotLocation, PivotLocationPixel);
+
+				const int32 NumSides = FUVEditorUXSettings::PivotCircleNumSides;
+				const float Radius = FUVEditorUXSettings::PivotCircleRadius;
+				const float LineThickness = FUVEditorUXSettings::PivotLineThickness;
+				const FColor LineColor = FUVEditorUXSettings::PivotLineColor;
+
+				const float	AngleDelta = 2.0f * UE_PI / NumSides;
+				FVector2D AxisX(1.f, 0.f);
+				FVector2D AxisY(0.f, -1.f);
+				FVector2D LastVertex = PivotLocationPixel + AxisX * Radius;
+
+				for (int32 SideIndex = 0; SideIndex < NumSides; SideIndex++)
+				{
+					const FVector2D Vertex = PivotLocationPixel + (AxisX * FMath::Cos(AngleDelta * (SideIndex + 1)) + AxisY * FMath::Sin(AngleDelta * (SideIndex + 1))) * Radius;
+					FCanvasLineItem LineItem(LastVertex, Vertex);
+					LineItem.LineThickness = LineThickness;
+					LineItem.SetColor(LineColor);
+					Canvas->DrawItem(LineItem);
+					LastVertex = Vertex;
+				}
+			}
+		}
+	}
 }
 
 
diff --git a/Engine/Plugins/Experimental/UVEditor/Source/UVEditorTools/Private/UVEditorUXSettings.cpp b/Engine/Plugins/Experimental/UVEditor/Source/UVEditorTools/Private/UVEditorUXSettings.cpp
--- a/Engine/Plugins/Experimental/UVEditor/Source/UVEditorTools/Private/UVEditorUXSettings.cpp
+++ b/Engine/Plugins/Experimental/UVEditor/Source/UVEditorTools/Private/UVEditorUXSettings.cpp
@@ -67,6 +67,8 @@ const FColor FUVEditorUXSettings::GridMajorColor(FColor::FromHex("#888888"));
 const FColor FUVEditorUXSettings::GridMinorColor(FColor::FromHex("#777777"));
 const FColor FUVEditorUXSettings::RulerXColor(FColor::FromHex("#888888"));
 const FColor FUVEditorUXSettings::RulerYColor(FColor::FromHex("#888888"));
+const FColor FUVEditorUXSettings::PivotLineColor(FColor::Cyan);
+
 // Thicknesses
 const float FUVEditorUXSettings::LivePreviewHighlightThickness(2.0);
 const float FUVEditorUXSettings::LivePreviewHighlightPointSize(4);
@@ -80,12 +82,17 @@ const float FUVEditorUXSettings::AxisThickness(2.0);
 const float FUVEditorUXSettings::GridMajorThickness(1.0);
 
 const float FUVEditorUXSettings::ToolPointSize(6);
+const float FUVEditorUXSettings::PivotLineThickness(1.5);
 
 // Grid
 const int32 FUVEditorUXSettings::GridSubdivisionsPerLevel(4);
 const int32 FUVEditorUXSettings::GridLevels(3);
 const int32 FUVEditorUXSettings::RulerSubdivisionLevel(1);
 
+// Pivot Visuals
+const int32 FUVEditorUXSettings::PivotCircleNumSides(32);
+const float FUVEditorUXSettings::PivotCircleRadius(10.0);
+
 // CVARs
 
 TAutoConsoleVariable<int32> FUVEditorUXSettings::CVarEnablePrototypeUDIMSupport(
diff --git a/Engine/Plugins/Experimental/UVEditor/Source/UVEditorTools/Public/ContextObjects/UVToolContextObjects.h b/Engine/Plugins/Experimental/UVEditor/Source/UVEditorTools/Public/ContextObjects/UVToolContextObjects.h
--- a/Engine/Plugins/Experimental/UVEditor/Source/UVEditorTools/Public/ContextObjects/UVToolContextObjects.h
+++ b/Engine/Plugins/Experimental/UVEditor/Source/UVEditorTools/Public/ContextObjects/UVToolContextObjects.h
@@ -318,3 +318,40 @@ protected:
 
 	TMap<FDynamicMesh3*, TreeInputObjectPair> AABBTreeStorage;
 };
+
+
+/**
+* Stores pointers to additional property sets for the active tool, used by the UVEditor Mode and Toolkit to 
+* populate other settings menus, such as the Display menu.
+*/
+
+UCLASS()
+class UVEDITORTOOLS_API UUVEditorToolPropertiesAPI : public UUVToolContextObject
+{
+	GENERATED_BODY()
+public:
+
+	/* Property sets passed here will be displayed in the Display menu at the top of the UV Editor. */
+	void SetToolDisplayProperties(UObject* ToolDisplayPropertesIn)
+	{
+		ToolDisplayProperties = ToolDisplayPropertesIn;
+	}
+
+	UObject* GetToolDisplayProperties() {
+		if (ToolDisplayProperties)
+		{
+			return ToolDisplayProperties;
+		}
+		return nullptr;
+	}
+
+	virtual void OnToolEnded(UInteractiveTool* DeadTool) {
+		ToolDisplayProperties = nullptr;
+	}
+
+protected:
+
+	UPROPERTY()
+	TObjectPtr<UObject> ToolDisplayProperties = nullptr;
+
+};
\ No newline at end of file
diff --git a/Engine/Plugins/Experimental/UVEditor/Source/UVEditorTools/Public/Operators/UVEditorUVTransformOp.h b/Engine/Plugins/Experimental/UVEditor/Source/UVEditorTools/Public/Operators/UVEditorUVTransformOp.h
--- a/Engine/Plugins/Experimental/UVEditor/Source/UVEditorTools/Public/Operators/UVEditorUVTransformOp.h
+++ b/Engine/Plugins/Experimental/UVEditor/Source/UVEditorTools/Public/Operators/UVEditorUVTransformOp.h
@@ -339,6 +339,7 @@ public:
 	void SetSelection(const TSet<int32>& TriangleSelection, const TSet<int32>& VertexSelection);
 
 	virtual void CalculateResult(FProgressCancel* Progress) override;
+	const TArray<FVector2D>& GetPivotLocations() const { return VisualizationPivots; }
 
 protected:
 	virtual void HandleTransformationOp(FProgressCancel* Progress) = 0;
@@ -356,6 +357,7 @@ protected:
 	TOptional< TSet<int32> > TransformingElements;
 	TSharedPtr<FMeshConnectedComponents> UVComponents;
 	TArray<int32> SpatiallyOrderedComponentIndex;
+	TArray<FVector2D> VisualizationPivots;
 
 	TOptional<TSet<int32>> TriangleSelection;
 	TOptional<TSet<int32>> VertexSelection;
diff --git a/Engine/Plugins/Experimental/UVEditor/Source/UVEditorTools/Public/UVEditorTransformTool.h b/Engine/Plugins/Experimental/UVEditor/Source/UVEditorTools/Public/UVEditorTransformTool.h
--- a/Engine/Plugins/Experimental/UVEditor/Source/UVEditorTools/Public/UVEditorTransformTool.h
+++ b/Engine/Plugins/Experimental/UVEditor/Source/UVEditorTools/Public/UVEditorTransformTool.h
@@ -19,6 +19,19 @@ class UUVEditorUVTransformOperatorFactory;
 class UCanvas;
 class UUVEditorTransformTool;
 
+/**
+ * Visualization settings for the TransformTool
+ */
+UCLASS()
+class UVEDITORTOOLS_API UUVEditorTransformToolDisplayProperties : public UInteractiveToolPropertySet
+{
+	GENERATED_BODY()
+public:
+	/** Draw the tool's active pivot location if needed.*/
+	UPROPERTY(EditAnywhere, Category = "Tool Display Settings", meta = (DisplayName = "Display Tool Pivots"))
+	bool bDrawPivots = true;
+};
+
 UCLASS()
 class UVEDITORTOOLS_API UUVEditorBaseTransformToolBuilder : public UInteractiveToolBuilder
 {
@@ -106,12 +119,17 @@ protected:
 	UPROPERTY()
 	TObjectPtr<UUVEditorUVTransformPropertiesBase> Settings = nullptr;
 
+	UPROPERTY()
+	TObjectPtr<UUVEditorTransformToolDisplayProperties> DisplaySettings = nullptr;
+
 	UPROPERTY()
 	TArray<TObjectPtr<UUVEditorUVTransformOperatorFactory>> Factories;
 
 	UPROPERTY()
 	TObjectPtr<UUVToolSelectionAPI> UVToolSelectionAPI = nullptr;
 
+	TArray<TArray<FVector2D>> PerTargetPivotLocations;
+
 	//
 	// Analytics
 	//
diff --git a/Engine/Plugins/Experimental/UVEditor/Source/UVEditorTools/Public/UVEditorUXSettings.h b/Engine/Plugins/Experimental/UVEditor/Source/UVEditorTools/Public/UVEditorUXSettings.h
--- a/Engine/Plugins/Experimental/UVEditor/Source/UVEditorTools/Public/UVEditorUXSettings.h
+++ b/Engine/Plugins/Experimental/UVEditor/Source/UVEditorTools/Public/UVEditorUXSettings.h
@@ -97,6 +97,12 @@ public:
 	static const int32 GridLevels;
 	static const int32 RulerSubdivisionLevel;
 
+	// Pivots
+	static const int32 PivotCircleNumSides;
+	static const float PivotCircleRadius;
+	static const float PivotLineThickness;
+	static const FColor PivotLineColor;
+
 	// Background
 	static const float BackgroundQuadDepthOffset;
 
