UVEditor: Fix bug where undoing a selection would continue to transform the previous selection using the interactive gizmo.

diff --git a/Engine/Plugins/Experimental/UVEditor/Source/UVEditorTools/Private/Selection/UVToolSelectionAPI.cpp b/Engine/Plugins/Experimental/UVEditor/Source/UVEditorTools/Private/Selection/UVToolSelectionAPI.cpp
--- a/Engine/Plugins/Experimental/UVEditor/Source/UVEditorTools/Private/Selection/UVToolSelectionAPI.cpp
+++ b/Engine/Plugins/Experimental/UVEditor/Source/UVEditorTools/Private/Selection/UVToolSelectionAPI.cpp
@@ -154,8 +154,6 @@ void UUVToolSelectionAPI::SetSelections(const TArray<FUVToolSelection>& Selectio
 		BeginChange();
 		bUsingExistingTransaction = false;
 	}
-	bool bWillCallEndChange = bEmitChange && !bUsingExistingTransaction;
-
 
 	bCachedUnwrapSelectionBoundingBoxCenterValid = false;
 
@@ -194,10 +192,25 @@ void UUVToolSelectionAPI::SetSelections(const TArray<FUVToolSelection>& Selectio
 		}
 	}
 
+	/*
+	* Depending on whether we're calling EndChangeAndEmitIfModified or not, the comparison of selections
+	* and the broadcasting of PreSelectionChange and OnSelectionChanged might be done for us. If we're not 
+	* calling EndChangeAndEmitIfModified (either because bEmitChange was false, or we're inside an existing
+	* transaction that we expect the user to close), then we have to do these things ourselves.
+	* 
+	* This is a bit messy, but it allows us to not duplicate the change emitting code from EndChangeAndEmitIfModified.
+	*/
+	bool bWillCallEndChange = bEmitChange && !bUsingExistingTransaction;
+
 	bool bSelectionsDiffer = false;
 	if (!bWillCallEndChange)
 	{
 		bSelectionsDiffer = DoSelectionSetsDiffer(CurrentSelections, NewSelections);
+
+		if (bBroadcast && bSelectionsDiffer)
+		{
+			OnPreSelectionChange.Broadcast(bEmitChange, (uint32)(ESelectionChangeTypeFlag::SelectionChanged));
+		}
 	}
 
 	CurrentSelections = MoveTemp(NewSelections);
@@ -206,6 +219,13 @@ void UUVToolSelectionAPI::SetSelections(const TArray<FUVToolSelection>& Selectio
 	{
 		bSelectionsDiffer = EndChangeAndEmitIfModified(bBroadcast);
 	}
+	else
+	{
+		if (bBroadcast && bSelectionsDiffer)
+		{
+			OnSelectionChanged.Broadcast(bEmitChange, (uint32)(ESelectionChangeTypeFlag::SelectionChanged));
+		}
+	}
 
 	if (bSelectionsDiffer)
 	{
@@ -281,7 +301,6 @@ void UUVToolSelectionAPI::SetUnsetElementAppliedMeshSelections(const TArray<FUVT
 		BeginChange();
 		bUsingExistingTransaction = false;
 	}
-	bool bWillCallEndChange = bEmitChange && !bUsingExistingTransaction;
 
 	// If we don't match with the current unset selection type, just clear them to keep everything consistent
 	if (CurrentSelections.Num() > 0 && UnsetSelectionsIn.Num() > 0 && UnsetSelectionsIn[0].Type != CurrentSelections[0].Type)
@@ -313,10 +332,25 @@ void UUVToolSelectionAPI::SetUnsetElementAppliedMeshSelections(const TArray<FUVT
 		}
 	}
 
+	/*
+	* Depending on whether we're calling EndChangeAndEmitIfModified or not, the comparison of selections
+	* and the broadcasting of PreSelectionChange and OnSelectionChanged might be done for us. If we're not 
+	* calling EndChangeAndEmitIfModified (either because bEmitChange was false, or we're inside an existing
+	* transaction that we expect the user to close), then we have to do these things ourselves.
+	* 
+	* This is a bit messy, but it allows us to not duplicate the change emitting code from EndChangeAndEmitIfModified.
+	*/
+	bool bWillCallEndChange = bEmitChange && !bUsingExistingTransaction;
+
 	bool bSelectionsDiffer = false;
 	if (!bWillCallEndChange)
 	{
 		bSelectionsDiffer = DoSelectionSetsDiffer(CurrentUnsetSelections, NewUnsetSelections);
+
+		if (bBroadcast && bSelectionsDiffer)
+		{
+			OnPreSelectionChange.Broadcast(bEmitChange, (uint32)(ESelectionChangeTypeFlag::UnsetSelectionChanged));
+		}
 	}
 
 	CurrentUnsetSelections = MoveTemp(NewUnsetSelections);
@@ -325,6 +359,13 @@ void UUVToolSelectionAPI::SetUnsetElementAppliedMeshSelections(const TArray<FUVT
 	{
 		bSelectionsDiffer = EndChangeAndEmitIfModified(bBroadcast);
 	}
+	else
+	{
+		if (bBroadcast && bSelectionsDiffer)
+		{
+			OnSelectionChanged.Broadcast(bEmitChange, (uint32)(ESelectionChangeTypeFlag::UnsetSelectionChanged));
+		}
+	}
 
 	if (bSelectionsDiffer)
 	{
