UVEditor: Clean up warnings related to type truncation and loss of precision issues.

diff --git a/Engine/Plugins/Experimental/UVEditor/Source/UVEditor/Private/UVEditor2DViewportClient.cpp b/Engine/Plugins/Experimental/UVEditor/Source/UVEditor/Private/UVEditor2DViewportClient.cpp
--- a/Engine/Plugins/Experimental/UVEditor/Source/UVEditor/Private/UVEditor2DViewportClient.cpp
+++ b/Engine/Plugins/Experimental/UVEditor/Source/UVEditor/Private/UVEditor2DViewportClient.cpp
@@ -50,9 +50,9 @@ namespace FUVEditor2DViewportClientLocals {
 		float RightPadding = 0.0;
 	};
 
-	FText ConvertToExponentialNotation(float InValue, float MinExponent = 0.0)
+	FText ConvertToExponentialNotation(double InValue, double MinExponent = 0.0)
 	{
-		float Exponent = FMath::LogX(10, InValue < 0 ? -InValue : InValue);
+		double Exponent = FMath::LogX(10, InValue < 0 ? -InValue : InValue);
 		Exponent = FMath::Floor(Exponent);
 		if (FMath::IsFinite(Exponent) && FMath::Abs(Exponent) > MinExponent)
 		{
@@ -132,7 +132,7 @@ namespace FUVEditor2DViewportClientLocals {
 
 	void ConvertPixelToUV(const FVector2D& PixelIn, double RelDepthIn, FVector2D& UVOut, const FSceneView& View)
 	{
-		FVector4 ScreenPoint = View.PixelToScreen(PixelIn.X, PixelIn.Y, RelDepthIn);
+		FVector4 ScreenPoint = View.PixelToScreen(static_cast<float>(PixelIn.X), static_cast<float>(PixelIn.Y), static_cast<float>(RelDepthIn));
 		FVector4 WorldPointHomogenous = View.ViewMatrices.GetInvViewProjectionMatrix().TransformFVector4(ScreenPoint);
 		FVector WorldPoint(WorldPointHomogenous.X / WorldPointHomogenous.W,
 			               WorldPointHomogenous.Y / WorldPointHomogenous.W,
@@ -251,12 +251,12 @@ void FUVEditor2DViewportClient::Draw(const FSceneView* View, FPrimitiveDrawInter
 void FUVEditor2DViewportClient::DrawGrid(const FSceneView* View, FPrimitiveDrawInterface* PDI)
 {
 	// Basic scaling amount
-	const float UVScale = UUVEditorMode::GetUVMeshScalingFactor();
+	const double UVScale = UUVEditorMode::GetUVMeshScalingFactor();
 	
 	// Determine important geometry of the viewport for creating grid lines
 	FVector WorldCenterPoint( 0,0,0 );
 	FVector4 WorldToScreenCenter = View->WorldToScreen(WorldCenterPoint);	
-	float ZoomFactor = WorldToScreenCenter.W;
+	double ZoomFactor = WorldToScreenCenter.W;
 	FVector4 MaxScreen(1 * ZoomFactor, 1 * ZoomFactor, 0, 1);
 	FVector4 MinScreen(-1 * ZoomFactor, -1 * ZoomFactor, 0, 1);
 	FVector WorldBoundsMax = View->ScreenToWorld(MaxScreen);
@@ -284,7 +284,7 @@ void FUVEditor2DViewportClient::DrawGrid(const FSceneView* View, FPrimitiveDrawI
 			PDI, Transform);
 	}
 
-	float AxisExtent = UVScale;
+	double AxisExtent = UVScale;
 
 	// Draw colored axis lines
 	PDI->DrawLine(FVector(0, 0, 0), FVector(AxisExtent, 0, 0), FUVEditorUXSettings::XAxisColor, SDPG_World, FUVEditorUXSettings::AxisThickness, 0, true);
@@ -460,23 +460,23 @@ void FUVEditor2DViewportClient::DrawCanvas(FViewport& InViewport, FSceneView& Vi
 
 void FUVEditor2DViewportClient::DrawGridRulers(FViewport& InViewport, FSceneView& View, UCanvas& Canvas)
 {
-	const float UVScale = FUVEditorUXSettings::UVMeshScalingFactor;
+	const double UVScale = FUVEditorUXSettings::UVMeshScalingFactor;
 	const int32 Level = FUVEditorUXSettings::RulerSubdivisionLevel;
 	const int32 Subdivisions = FUVEditorUXSettings::GridSubdivisionsPerLevel; // Number of subdivisions per level
 
 	// Determine important geometry of the viewport for creating grid lines
 	FVector WorldCenterPoint(0, 0, 0);
 	FVector4 WorldToScreenCenter = View.WorldToScreen(WorldCenterPoint);
-	float ZoomFactor = WorldToScreenCenter.W;
-	float GridZoomFactor = ZoomFactor / UVScale;
+	double ZoomFactor = WorldToScreenCenter.W;
+	double GridZoomFactor = ZoomFactor / UVScale;
 
-	float LogZoom = FMath::LogX(Subdivisions, GridZoomFactor);
-	float LogZoomDirection = FMath::Sign(LogZoom);
+	double LogZoom = FMath::LogX(Subdivisions, GridZoomFactor);
+	double LogZoomDirection = FMath::Sign(LogZoom);
 	LogZoom = FMath::Abs(LogZoom);
-	LogZoom = FMathf::Floor(LogZoom);
+	LogZoom = FMath::Floor(LogZoom);
 	LogZoom = LogZoomDirection * LogZoom;
 
-	double GridScale = FMathf::Pow(Subdivisions, LogZoom - Level);
+	double GridScale = FMathd::Pow(Subdivisions, LogZoom - static_cast<double>(Level));
 
 	FVector2D PixelMinBounds(0, Canvas.SizeY);
 	FVector2D PixelMaxBounds(Canvas.SizeX, 0);
diff --git a/Engine/Plugins/Experimental/UVEditor/Source/UVEditor/Private/UVEditorBackgroundPreview.cpp b/Engine/Plugins/Experimental/UVEditor/Source/UVEditor/Private/UVEditorBackgroundPreview.cpp
--- a/Engine/Plugins/Experimental/UVEditor/Source/UVEditor/Private/UVEditorBackgroundPreview.cpp
+++ b/Engine/Plugins/Experimental/UVEditor/Source/UVEditor/Private/UVEditorBackgroundPreview.cpp
@@ -195,8 +195,8 @@ void UUVEditorBackgroundPreview::UpdateBackground()
 			FVector2f Block;
 			int32 UCoord, VCoord;
 			UE::TextureUtilitiesCommon::ExtractUDIMCoordinates(Settings->UDIMBlocks[BlockIndex], UCoord, VCoord);
-			Block.X = UCoord;
-			Block.Y = VCoord;
+			Block.X = static_cast<float>(UCoord);
+			Block.Y = static_cast<float>(VCoord);
 			UDimBlocksToRender.Push(Block);
 		}
 	}
diff --git a/Engine/Plugins/Experimental/UVEditor/Source/UVEditor/Private/UVEditorMode.cpp b/Engine/Plugins/Experimental/UVEditor/Source/UVEditor/Private/UVEditorMode.cpp
--- a/Engine/Plugins/Experimental/UVEditor/Source/UVEditor/Private/UVEditorMode.cpp
+++ b/Engine/Plugins/Experimental/UVEditor/Source/UVEditor/Private/UVEditorMode.cpp
@@ -1257,7 +1257,7 @@ void UUVEditorMode::UpdatePreviewMaterialBasedOnBackground()
 		for (int32 AssetID = 0; AssetID < ToolInputObjects.Num(); ++AssetID) {			
 			TObjectPtr<UMaterialInstanceDynamic> BackgroundMaterialOverride = UMaterialInstanceDynamic::Create(BackgroundVisualization->BackgroundMaterial->Parent, this);
 			BackgroundMaterialOverride->CopyInterpParameters(BackgroundVisualization->BackgroundMaterial);
-			BackgroundMaterialOverride->SetScalarParameterValue(TEXT("UVChannel"), GetDisplayedChannel(AssetID));
+			BackgroundMaterialOverride->SetScalarParameterValue(TEXT("UVChannel"), static_cast<float>(GetDisplayedChannel(AssetID)));
 
 			ToolInputObjects[AssetID]->AppliedPreview->OverrideMaterial = BackgroundMaterialOverride;
 		}
diff --git a/Engine/Plugins/Experimental/UVEditor/Source/UVEditorTools/Private/Operators/UVEditorRecomputeUVsOp.cpp b/Engine/Plugins/Experimental/UVEditor/Source/UVEditorTools/Private/Operators/UVEditorRecomputeUVsOp.cpp
--- a/Engine/Plugins/Experimental/UVEditor/Source/UVEditorTools/Private/Operators/UVEditorRecomputeUVsOp.cpp
+++ b/Engine/Plugins/Experimental/UVEditor/Source/UVEditorTools/Private/Operators/UVEditorRecomputeUVsOp.cpp
@@ -581,7 +581,7 @@ TUniquePtr<FDynamicMeshOperator> UUVEditorRecomputeUVsOpFactory::MakeNewOperator
 		break;
 	case EUVEditorRecomputeUVsPropertiesLayoutType::NormalizeToBounds:
 		RecomputeUVsOp->bNormalizeAreas = true;
-		RecomputeUVsOp->AreaScaling = Settings->NormalizeScale / OriginalMesh->GetBounds().MaxDim();
+		RecomputeUVsOp->AreaScaling = Settings->NormalizeScale / static_cast<float>(OriginalMesh->GetBounds().MaxDim());
 		break;
 	case EUVEditorRecomputeUVsPropertiesLayoutType::NormalizeToWorld:
 		RecomputeUVsOp->bNormalizeAreas = true;
diff --git a/Engine/Plugins/Experimental/UVEditor/Source/UVEditorTools/Private/Operators/UVEditorUVTransformOp.cpp b/Engine/Plugins/Experimental/UVEditor/Source/UVEditorTools/Private/Operators/UVEditorUVTransformOp.cpp
--- a/Engine/Plugins/Experimental/UVEditor/Source/UVEditorTools/Private/Operators/UVEditorUVTransformOp.cpp
+++ b/Engine/Plugins/Experimental/UVEditor/Source/UVEditorTools/Private/Operators/UVEditorUVTransformOp.cpp
@@ -126,22 +126,22 @@ FVector2f FUVEditorUVTransformBaseOp::GetAlignmentPointFromBoundingBoxAndDirecti
 		AlignmentPoint = FVector2f(0,0);
 		break;
 	case EUVEditorAlignDirectionBackend::Top:
-		AlignmentPoint = FVector2f(BoundingBox.Center().X, BoundingBox.Max.Y);
+		AlignmentPoint = FVector2f(static_cast<float>(BoundingBox.Center().X), static_cast<float>(BoundingBox.Max.Y));
 		break;
 	case EUVEditorAlignDirectionBackend::Bottom:
-		AlignmentPoint = FVector2f(BoundingBox.Center().X, BoundingBox.Min.Y);
+		AlignmentPoint = FVector2f(static_cast<float>(BoundingBox.Center().X), static_cast<float>(BoundingBox.Min.Y));
 		break;
 	case EUVEditorAlignDirectionBackend::Left:
-		AlignmentPoint = FVector2f(BoundingBox.Min.X, BoundingBox.Center().Y);
+		AlignmentPoint = FVector2f(static_cast<float>(BoundingBox.Min.X), static_cast<float>(BoundingBox.Center().Y));
 		break;
 	case EUVEditorAlignDirectionBackend::Right:
-		AlignmentPoint = FVector2f(BoundingBox.Max.X, BoundingBox.Center().Y);
+		AlignmentPoint = FVector2f(static_cast<float>(BoundingBox.Max.X), static_cast<float>(BoundingBox.Center().Y));
 		break;
 	case EUVEditorAlignDirectionBackend::CenterVertically:
-		AlignmentPoint = FVector2f(BoundingBox.Center().X, BoundingBox.Center().Y);
+		AlignmentPoint = FVector2f(static_cast<float>(BoundingBox.Center().X), static_cast<float>(BoundingBox.Center().Y));
 		break;
 	case EUVEditorAlignDirectionBackend::CenterHorizontally:
-		AlignmentPoint = FVector2f(BoundingBox.Center().X, BoundingBox.Center().Y);
+		AlignmentPoint = FVector2f(static_cast<float>(BoundingBox.Center().X), static_cast<float>(BoundingBox.Center().Y));
 		break;
 	default:
 		ensure(false);
@@ -294,8 +294,8 @@ void FUVEditorUVTransformOp::HandleTransformationOp(FProgressCancel * Progress)
 		FVector2f ScalePivot = GetPivotFromMode(Vid, PivotMode);
 		FVector2f UV = TransformOpLocals::UnwrapPositionToUV(ResultMesh->GetVertexRef(Vid));
 		UV = (UV - ScalePivot);
-		UV[0] *= Scale[0];
-		UV[1] *= Scale[1];
+		UV[0] *= static_cast<float>(Scale[0]);
+		UV[1] *= static_cast<float>(Scale[1]);
 		UV = (UV + ScalePivot);
 		ResultMesh->SetVertex(Vid, TransformOpLocals::UVToUnwrapPosition(UV));
 	};
@@ -308,8 +308,8 @@ void FUVEditorUVTransformOp::HandleTransformationOp(FProgressCancel * Progress)
 		// where positive values are clockwise and negative values are counterclockwise.
 		double RotationInRadians = -RotationIn / 180.0 * UE_PI;
 		UV = (UV - Pivot);
-		UV_Rotated[0] = UV[0] * FMath::Cos(RotationInRadians) - UV[1] * FMath::Sin(RotationInRadians);
-		UV_Rotated[1] = UV[0] * FMath::Sin(RotationInRadians) + UV[1] * FMath::Cos(RotationInRadians);
+		UV_Rotated[0] = UV[0] * static_cast<float>(FMath::Cos(RotationInRadians)) - UV[1] * static_cast<float>(FMath::Sin(RotationInRadians));
+		UV_Rotated[1] = UV[0] * static_cast<float>(FMath::Sin(RotationInRadians)) + UV[1] * static_cast<float>(FMath::Cos(RotationInRadians));
 		UV = (UV_Rotated + Pivot);
 		ResultMesh->SetVertex(Vid, TransformOpLocals::UVToUnwrapPosition(UV));
 	};
@@ -582,9 +582,9 @@ void FUVEditorUVDistributeOp::HandleTransformationOp(FProgressCancel* Progress)
 		float TotalDistance = 0.0f;
 		for (int32 Cid = 0; Cid < NumComponents; ++Cid)
 		{
-			TotalDistance += bVertical ? PerComponentBoundingBoxes[Cid].Height() : PerComponentBoundingBoxes[Cid].Width();
+			TotalDistance += static_cast<float>( bVertical ? PerComponentBoundingBoxes[Cid].Height() : PerComponentBoundingBoxes[Cid].Width());
 		}
-		float BoundingBoxDistance = bVertical ? OverallBoundingBox.Height() : OverallBoundingBox.Width();
+		float BoundingBoxDistance = static_cast<float>( bVertical ? OverallBoundingBox.Height() : OverallBoundingBox.Width());
 		if (bEnableManualDistances)
 		{
 			BoundingBoxDistance = ManualExtent;
@@ -623,7 +623,7 @@ void FUVEditorUVDistributeOp::HandleTransformationOp(FProgressCancel* Progress)
 			}
 			if (bEqualizeSpacing)
 			{
-				float ComponentSpace = bVertical ? PerComponentBoundingBoxes[Cid].Height() : PerComponentBoundingBoxes[Cid].Width();
+				float ComponentSpace = static_cast<float>( bVertical ? PerComponentBoundingBoxes[Cid].Height() : PerComponentBoundingBoxes[Cid].Width());
 				NextPosition += SpreadDirection * (ComponentSpace + GapSpace);
 			}
 			else
diff --git a/Engine/Plugins/Experimental/UVEditor/Source/UVEditorTools/Private/Selection/UVToolSelectionHighlightMechanic.cpp b/Engine/Plugins/Experimental/UVEditor/Source/UVEditorTools/Private/Selection/UVToolSelectionHighlightMechanic.cpp
--- a/Engine/Plugins/Experimental/UVEditor/Source/UVEditorTools/Private/Selection/UVToolSelectionHighlightMechanic.cpp
+++ b/Engine/Plugins/Experimental/UVEditor/Source/UVEditorTools/Private/Selection/UVToolSelectionHighlightMechanic.cpp
@@ -24,6 +24,15 @@ using namespace UE::Geometry;
 
 namespace UVToolSelectionHighlightMechanicLocals
 {
+	FVector2f ToFVector2f(const FVector& VectorIn)
+    {
+         return FVector2f( static_cast<float>(VectorIn.X), static_cast<float>(VectorIn.Y) );
+    }
+
+	FVector3f ToFVector3f(const FVector3d& VectorIn)
+    {
+         return FVector3f( static_cast<float>(VectorIn.X), static_cast<float>(VectorIn.Y), static_cast<float>(VectorIn.Z));
+    }
 }
 
 void UUVToolSelectionHighlightMechanic::Initialize(UWorld* UnwrapWorld, UWorld* LivePreviewWorld)
@@ -46,7 +55,7 @@ void UUVToolSelectionHighlightMechanic::Initialize(UWorld* UnwrapWorld, UWorld*
 	UnwrapTriangleSet = NewObject<UBasic2DTriangleSetComponent>(UnwrapGeometryActor);
 	// We are setting the TranslucencySortPriority here to handle the UV editor's use case in 2D
 	// where multiple translucent layers are drawn on top of each other but still need depth sorting.
-	UnwrapTriangleSet->TranslucencySortPriority = FUVEditorUXSettings::SelectionTriangleDepthBias;
+	UnwrapTriangleSet->TranslucencySortPriority = static_cast<int32>(FUVEditorUXSettings::SelectionTriangleDepthBias);
 	TriangleSetMaterial = ToolSetupUtil::GetCustomTwoSidedDepthOffsetMaterial(GetParentTool()->GetToolManager(),
 		FUVEditorUXSettings::SelectionTriangleFillColor,
 		FUVEditorUXSettings::SelectionTriangleDepthBias,
@@ -174,6 +183,8 @@ void UUVToolSelectionHighlightMechanic::RebuildUnwrapHighlight(
 	const TArray<FUVToolSelection>& Selections, const FTransform& StartTransform, 
 	bool bUsePreviews)
 {
+    using namespace UVToolSelectionHighlightMechanicLocals;
+
 	if (!ensure(UnwrapGeometryActor))
 	{
 		return;
@@ -220,15 +231,12 @@ void UUVToolSelectionHighlightMechanic::RebuildUnwrapHighlight(
 				{
 					Points[i] = StartTransform.InverseTransformPosition(Points[i]);
 				}
-				UnwrapTriangleSetPtr->AddElement(FVector2f(Points[0].X, Points[0].Y),
-											  FVector2f(Points[1].X, Points[1].Y),
-											  FVector2f(Points[2].X, Points[2].Y));
+				UnwrapTriangleSetPtr->AddElement(ToFVector2f(Points[0]),ToFVector2f(Points[1]),ToFVector2f(Points[2]));
 					
 				for (int i = 0; i < 3; ++i)
 				{
 					int NextIndex = (i + 1) % 3;
-					UnwrapLineSetPtr->AddElement(FVector2f(Points[i].X, Points[i].Y),
-						                      FVector2f(Points[NextIndex].X, Points[NextIndex].Y));
+					UnwrapLineSetPtr->AddElement(ToFVector2f(Points[i]),ToFVector2f(Points[NextIndex]));
 				}
 			}
 			UnwrapTriangleSetPtr->MarkRenderStateDirty();
@@ -264,8 +272,7 @@ void UUVToolSelectionHighlightMechanic::RebuildUnwrapHighlight(
 				Mesh.GetEdgeV(Eid, Points[0], Points[1]);
 				Points[0] = StartTransform.InverseTransformPosition(Points[0]);
 				Points[1] = StartTransform.InverseTransformPosition(Points[1]);
-				UnwrapLineSetPtr->AddElement(FVector2f(Points[0].X, Points[0].Y),
-					                      FVector2f(Points[1].X, Points[1].Y));
+				UnwrapLineSetPtr->AddElement(ToFVector2f(Points[0]),ToFVector2f(Points[1]));
 
 				if (bPairedEdgeHighlightsEnabled)
 				{
@@ -286,8 +293,7 @@ void UUVToolSelectionHighlightMechanic::RebuildUnwrapHighlight(
 					Mesh.GetEdgeV(Eid, Points[0], Points[1]);
 					Points[0] = StartTransform.InverseTransformPosition(Points[0]);
 					Points[1] = StartTransform.InverseTransformPosition(Points[1]);
-					SewEdgePairingLeftLineSetPtr->AddElement(FVector2f(Points[0].X, Points[0].Y),
-					                                      FVector2f(Points[1].X, Points[1].Y));
+					SewEdgePairingLeftLineSetPtr->AddElement(ToFVector2f(Points[0]),ToFVector2f(Points[1]));
 
 					// The paired edge may need to go into a separate line set if it is not selected so that it does
 					// not get affected by transformations of the selected highlights in SetUnwrapHighlightTransform
@@ -297,14 +303,13 @@ void UUVToolSelectionHighlightMechanic::RebuildUnwrapHighlight(
 					{
 						Points[0] = StartTransform.InverseTransformPosition(Points[0]);
 						Points[1] = StartTransform.InverseTransformPosition(Points[1]);
-						SewEdgePairingRightLineSetPtr->AddElement(FVector2f(Points[0].X, Points[0].Y),
-							                                      FVector2f(Points[1].X, Points[1].Y));
+						SewEdgePairingRightLineSetPtr->AddElement(ToFVector2f(Points[0]),ToFVector2f(Points[1]));
 					}
 					else
 					{
 						StaticPairedEdgeVidsPerMesh.Last().Value.Add(TPair<int32, int32>(Vids2.A, Vids2.B));
-						SewEdgeUnselectedPairingLineSetPtr->AddElement(FVector2f(Points[0].X, Points[0].Y),
- 							                                           FVector2f(Points[1].X, Points[1].Y));						
+						SewEdgeUnselectedPairingLineSetPtr->AddElement(ToFVector2f(Points[0]),
+ 							                                           ToFVector2f(Points[1]));
 					}
 				}//end if visualizing paired edges
 			}//end for each edge
@@ -328,7 +333,7 @@ void UUVToolSelectionHighlightMechanic::RebuildUnwrapHighlight(
 				}
 
 				const FVector3d Position = StartTransform.InverseTransformPosition(Mesh.GetVertex(Vid));
-				UnwrapPointSetPtr->AddElement(FVector2f(Position.X, Position.Y));
+				UnwrapPointSetPtr->AddElement(ToFVector2f(Position));
 			}
 			UnwrapPointSetPtr->MarkRenderStateDirty();
 		}
@@ -338,6 +343,8 @@ void UUVToolSelectionHighlightMechanic::RebuildUnwrapHighlight(
 void UUVToolSelectionHighlightMechanic::SetUnwrapHighlightTransform(const FTransform& Transform, 
 	bool bRebuildStaticPairedEdges, bool bUsePreviews)
 {
+    using namespace UVToolSelectionHighlightMechanicLocals;
+
 	if (ensure(UnwrapGeometryActor))
 	{
 		UnwrapGeometryActor->SetActorTransform(Transform);
@@ -375,8 +382,7 @@ void UUVToolSelectionHighlightMechanic::SetUnwrapHighlightTransform(const FTrans
 				}
 				FVector3d VertA = Mesh.GetVertex(VidPair.Key);
 				FVector3d VertB = Mesh.GetVertex(VidPair.Value);
-				SewEdgeUnselectedPairingLineSetPtr->AddElement(FVector2f(VertA.X, VertA.Y),
-                                         					FVector2f(VertB.X, VertB.Y));
+				SewEdgeUnselectedPairingLineSetPtr->AddElement(ToFVector2f(VertA),ToFVector2f(VertB));
 			}
 			SewEdgeUnselectedPairingLineSetPtr->MarkRenderStateDirty();
 		}
@@ -395,6 +401,7 @@ FTransform UUVToolSelectionHighlightMechanic::GetUnwrapHighlightTransform()
 void UUVToolSelectionHighlightMechanic::RebuildAppliedHighlightFromUnwrapSelection(
 	const TArray<FUVToolSelection>& UnwrapSelections, bool bUsePreviews)
 {
+    using namespace UVToolSelectionHighlightMechanicLocals;
 	if (!ensure(LivePreviewGeometryActor))
 	{
 		return;
@@ -428,9 +435,7 @@ void UUVToolSelectionHighlightMechanic::RebuildAppliedHighlightFromUnwrapSelecti
 			AppliedMesh.GetEdgeV(AppliedEid, Points[0], Points[1]);
 			Points[0] = MeshTransform.TransformPosition(Points[0]);
 			Points[1] = MeshTransform.TransformPosition(Points[1]);
-			LivePreviewLineSetPtr->AddElement(
-				FVector3f(Points[0].X, Points[0].Y, Points[0].Z),
-				FVector3f(Points[1].X, Points[1].Y, Points[1].Z));
+			LivePreviewLineSetPtr->AddElement(ToFVector3f(Points[0]),ToFVector3f(Points[1]));
 		};
 
 		if (Selection.Type == FUVToolSelection::EType::Triangle)
@@ -502,6 +507,7 @@ void UUVToolSelectionHighlightMechanic::RebuildAppliedHighlightFromUnwrapSelecti
 
 void UUVToolSelectionHighlightMechanic::AppendAppliedHighlight(const TArray<FUVToolSelection>& AppliedSelections, bool bUsePreviews)
 {
+    using namespace UVToolSelectionHighlightMechanicLocals;
 	if (!ensure(LivePreviewGeometryActor))
 	{
 		return;
@@ -530,9 +536,7 @@ void UUVToolSelectionHighlightMechanic::AppendAppliedHighlight(const TArray<FUVT
 			AppliedMesh.GetEdgeV(AppliedEid, Points[0], Points[1]);
 			Points[0] = MeshTransform.TransformPosition(Points[0]);
 			Points[1] = MeshTransform.TransformPosition(Points[1]);
-			LivePreviewLineSetPtr->AddElement(
-				FVector3f(Points[0].X, Points[0].Y, Points[0].Z),
-				FVector3f(Points[1].X, Points[1].Y, Points[1].Z));
+			LivePreviewLineSetPtr->AddElement(ToFVector3f(Points[0]),ToFVector3f(Points[1]));
 		};
 
 		if (Selection.Type == FUVToolSelection::EType::Triangle)
diff --git a/Engine/Plugins/Experimental/UVEditor/Source/UVEditorTools/Private/UVEditorUXSettings.cpp b/Engine/Plugins/Experimental/UVEditor/Source/UVEditorTools/Private/UVEditorUXSettings.cpp
--- a/Engine/Plugins/Experimental/UVEditor/Source/UVEditorTools/Private/UVEditorUXSettings.cpp
+++ b/Engine/Plugins/Experimental/UVEditor/Source/UVEditorTools/Private/UVEditorUXSettings.cpp
@@ -106,7 +106,7 @@ FLinearColor FUVEditorUXSettings::GetTriangleColorByTargetIndex(int32 TargetInde
 	double GoldenAngle = 137.50776405;
 
 	FLinearColor BaseColorHSV = FLinearColor::FromSRGBColor(UnwrapTriangleFillColor).LinearRGBToHSV();
-	BaseColorHSV.R = FMath::Fmod(BaseColorHSV.R + (GoldenAngle / 2.0 * TargetIndex), 360);;
+	BaseColorHSV.R = static_cast<float>(FMath::Fmod(BaseColorHSV.R + (GoldenAngle / 2.0 * TargetIndex), 360));
 
 	return BaseColorHSV.HSVToLinearRGB();
 }
@@ -143,7 +143,7 @@ FVector3d FUVEditorUXSettings::UVToVertPosition(const FVector2f& UV)
 
 FVector2f FUVEditorUXSettings::VertPositionToUV(const FVector3d& VertPosition)
 {
-	return FVector2f(VertPosition.Y / UVMeshScalingFactor, 1 - (VertPosition.X / UVMeshScalingFactor));
+	return FVector2f(static_cast<float>(VertPosition.Y) / UVMeshScalingFactor, 1.0f - (static_cast<float>(VertPosition.X) / UVMeshScalingFactor));
 };
 
 float FUVEditorUXSettings::LocationSnapValue(int32 LocationSnapMenuIndex)
diff --git a/Engine/Plugins/Experimental/UVEditor/Source/UVEditorToolsEditorOnly/Private/DetailsCustomizations/UVUnwrapToolCustomizations.cpp b/Engine/Plugins/Experimental/UVEditor/Source/UVEditorToolsEditorOnly/Private/DetailsCustomizations/UVUnwrapToolCustomizations.cpp
--- a/Engine/Plugins/Experimental/UVEditor/Source/UVEditorToolsEditorOnly/Private/DetailsCustomizations/UVUnwrapToolCustomizations.cpp
+++ b/Engine/Plugins/Experimental/UVEditor/Source/UVEditorToolsEditorOnly/Private/DetailsCustomizations/UVUnwrapToolCustomizations.cpp
@@ -88,7 +88,7 @@ void FUVEditorRecomputeUVsToolDetails::EvaluateLayoutTypeRestrictions(TSharedRef
 		(IslandMode == EUVEditorRecomputeUVsPropertiesIslandMode::PolyGroups ||
 		 UnwrapType == EUVEditorRecomputeUVsPropertiesUnwrapType::IslandMerging))
 	{
-		uint8 EnumValue = ImportTypeEnum->GetValueByName("EUVEditorRecomputeUVsPropertiesLayoutType::NormalizeToExistingBounds");
+		uint8 EnumValue = static_cast<uint8>(ImportTypeEnum->GetValueByName("EUVEditorRecomputeUVsPropertiesLayoutType::NormalizeToExistingBounds"));
 		LayoutTypeEnumRestriction->AddDisabledValue("NormalizeToExistingBounds");
 
 		uint8 Value = 0;
diff --git a/Engine/Plugins/Experimental/UVEditor/Source/UVEditorToolsEditorOnly/Private/Operators/UVEditorParameterizeMeshOp.cpp b/Engine/Plugins/Experimental/UVEditor/Source/UVEditorToolsEditorOnly/Private/Operators/UVEditorParameterizeMeshOp.cpp
--- a/Engine/Plugins/Experimental/UVEditor/Source/UVEditorToolsEditorOnly/Private/Operators/UVEditorParameterizeMeshOp.cpp
+++ b/Engine/Plugins/Experimental/UVEditor/Source/UVEditorToolsEditorOnly/Private/Operators/UVEditorParameterizeMeshOp.cpp
@@ -85,8 +85,7 @@ FUVEditorParameterizeMeshOp::FLinearMesh::FLinearMesh(const FDynamicMesh3& Mesh,
 
 		for (const FVector3d& Vertex : Mesh.VerticesItr())
 		{
-			FVector3f Pos(Vertex.X, Vertex.Y, Vertex.Z);
-			Positions.Add(Pos);
+			Positions.Add(static_cast<FVector3f>(Vertex));
 		}
 	}
 
diff --git a/Engine/Source/Editor/Experimental/EditorInteractiveToolsFramework/Private/Behaviors/2DViewportBehaviorTargets.cpp b/Engine/Source/Editor/Experimental/EditorInteractiveToolsFramework/Private/Behaviors/2DViewportBehaviorTargets.cpp
--- a/Engine/Source/Editor/Experimental/EditorInteractiveToolsFramework/Private/Behaviors/2DViewportBehaviorTargets.cpp
+++ b/Engine/Source/Editor/Experimental/EditorInteractiveToolsFramework/Private/Behaviors/2DViewportBehaviorTargets.cpp
@@ -13,7 +13,7 @@ FInputRayHit FEditor2DScrollBehaviorTarget::CanBeginClickDragSequence(const FInp
 	// Verify that ray is facing the proper direction
 	if (PressPos.WorldRay.Direction.Z * PressPos.WorldRay.Origin.Z < 0)
 	{
-		return FInputRayHit(-PressPos.WorldRay.Origin.Z / PressPos.WorldRay.Direction.Z);
+		return FInputRayHit(static_cast<float>(-PressPos.WorldRay.Origin.Z / PressPos.WorldRay.Direction.Z));
 	}
 	return FInputRayHit();
 }
@@ -87,8 +87,8 @@ void FEditor2DMouseWheelZoomBehaviorTarget::OnMouseWheelScrollUp(const FInputDev
 
 	FVector NewLocation = OriginalLocation + (ZoomInProportion * TToPlane * CurrentPos.WorldRay.Direction);
 
-	ViewportClient->OverrideFarClipPlane(NewLocation.Z - CameraFarPlaneWorldZ);
-	ViewportClient->OverrideNearClipPlane(NewLocation.Z * (1.0-CameraNearPlaneProportionZ));
+	ViewportClient->OverrideFarClipPlane(static_cast<float>(NewLocation.Z - CameraFarPlaneWorldZ));
+	ViewportClient->OverrideNearClipPlane(static_cast<float>(NewLocation.Z * (1.0-CameraNearPlaneProportionZ)));
 
 	// Don't zoom in so far that the XY plane lies in front of our near clipping plane, or else everything
 	// will suddenly disappear.
@@ -104,8 +104,8 @@ void FEditor2DMouseWheelZoomBehaviorTarget::OnMouseWheelScrollDown(const FInputD
 	double TToPlane = -OriginalLocation.Z / CurrentPos.WorldRay.Direction.Z;
 	FVector NewLocation = OriginalLocation - (ZoomOutProportion * TToPlane * CurrentPos.WorldRay.Direction);
 
-	ViewportClient->OverrideFarClipPlane(NewLocation.Z - CameraFarPlaneWorldZ);
-	ViewportClient->OverrideNearClipPlane(NewLocation.Z * (1.0 - CameraNearPlaneProportionZ));
+	ViewportClient->OverrideFarClipPlane(static_cast<float>(NewLocation.Z - CameraFarPlaneWorldZ));
+	ViewportClient->OverrideNearClipPlane(static_cast<float>(NewLocation.Z * (1.0 - CameraNearPlaneProportionZ)));
 
 	if (NewLocation.Z < ZoomOutLimit)
 	{
