UVEditor: Revert prior change to Transform tool's absolute translation mode to return to prior, correct behavior.

diff --git a/Engine/Plugins/Experimental/UVEditor/Source/UVEditorTools/Private/Operators/UVEditorUVTransformOp.cpp b/Engine/Plugins/Experimental/UVEditor/Source/UVEditorTools/Private/Operators/UVEditorUVTransformOp.cpp
--- a/Engine/Plugins/Experimental/UVEditor/Source/UVEditorTools/Private/Operators/UVEditorUVTransformOp.cpp
+++ b/Engine/Plugins/Experimental/UVEditor/Source/UVEditorTools/Private/Operators/UVEditorUVTransformOp.cpp
@@ -341,7 +341,7 @@ void FUVEditorUVTransformOp::HandleTransformationOp(FProgressCancel * Progress)
 			RotationPivot = GetPivotFromMode(Vid, PivotMode);
 		}				
 		FVector2f UV = TransformOpLocals::UnwrapPositionToUV(ResultMesh->GetVertexRef(Vid));
-		UV = (UV + FVector2f(Translation) + RotationPivot);
+		UV = (UV + FVector2f(Translation) - RotationPivot);
 		ResultMesh->SetVertex(Vid, TransformOpLocals::UVToUnwrapPosition(UV));
 	};
 
