UVEditor: Fix bug where select tool's gizmo would be in the incorrect place immediately after a tool completes or an undo occurred and the selected elements were moved as part of it.

diff --git a/Engine/Plugins/Experimental/UVEditor/Source/UVEditorTools/Private/Selection/UVToolSelectionAPI.cpp b/Engine/Plugins/Experimental/UVEditor/Source/UVEditorTools/Private/Selection/UVToolSelectionAPI.cpp
--- a/Engine/Plugins/Experimental/UVEditor/Source/UVEditorTools/Private/Selection/UVToolSelectionAPI.cpp
+++ b/Engine/Plugins/Experimental/UVEditor/Source/UVEditorTools/Private/Selection/UVToolSelectionAPI.cpp
@@ -77,10 +77,14 @@ void UUVToolSelectionAPI::SetTargets(const TArray<TObjectPtr<UUVEditorToolMeshIn
 			{
 				FUVToolSelection* TargetSelection = CurrentSelections.FindByPredicate(
 					[Target](const FUVToolSelection& CandidateSelection) { return CandidateSelection.Target == Target; });
-				if (TargetSelection && TargetSelection->Type == FUVToolSelection::EType::Edge)
+				if (TargetSelection)
 				{
-					TargetSelection->RestoreFromStableEdgeIdentifiers(*Target->UnwrapCanonical);
-				}
+					bCachedUnwrapSelectionBoundingBoxCenterValid = false;
+					if (TargetSelection->Type == FUVToolSelection::EType::Edge)
+					{
+						TargetSelection->RestoreFromStableEdgeIdentifiers(*Target->UnwrapCanonical);
+					}
+				}				
 			});
 	}
 }
diff --git a/Engine/Plugins/Experimental/UVEditor/Source/UVEditorTools/Private/UVSelectTool.cpp b/Engine/Plugins/Experimental/UVEditor/Source/UVEditorTools/Private/UVSelectTool.cpp
--- a/Engine/Plugins/Experimental/UVEditor/Source/UVEditorTools/Private/UVSelectTool.cpp
+++ b/Engine/Plugins/Experimental/UVEditor/Source/UVEditorTools/Private/UVSelectTool.cpp
@@ -224,7 +224,9 @@ void UUVSelectTool::Setup()
 		{
 			if (bUpdateGizmoOnCanonicalChange) // Used to avoid reacting to broadcasts that we ourselves caused via gizmo movement
 			{
-				UpdateGizmo();
+				// We update the gizmo with 'true' here to have the selection api recompute the
+				// centroid because it may not yet be notified that the mesh has changed.
+				UpdateGizmo(true);
 				SelectionAPI->RebuildUnwrapHighlight(TransformGizmo->GetGizmoTransform());
 			}
 		});
@@ -339,11 +341,11 @@ void UUVSelectTool::OnPropertyModified(UObject* PropertySet, FProperty* Property
 {
 }
 
-void UUVSelectTool::UpdateGizmo()
+void UUVSelectTool::UpdateGizmo(bool bForceRecomputeSelectionCenters)
 {
 	if (SelectionAPI->HaveSelections())
 	{
-		TransformGizmo->ReinitializeGizmoTransform(FTransform(SelectionAPI->GetUnwrapSelectionBoundingBoxCenter()));
+		TransformGizmo->ReinitializeGizmoTransform(FTransform(SelectionAPI->GetUnwrapSelectionBoundingBoxCenter(bForceRecomputeSelectionCenters)));
 	}
 
 	TransformGizmo->SetVisibility(
diff --git a/Engine/Plugins/Experimental/UVEditor/Source/UVEditorTools/Public/Selection/UVToolSelectionAPI.h b/Engine/Plugins/Experimental/UVEditor/Source/UVEditorTools/Public/Selection/UVToolSelectionAPI.h
--- a/Engine/Plugins/Experimental/UVEditor/Source/UVEditorTools/Public/Selection/UVToolSelectionAPI.h
+++ b/Engine/Plugins/Experimental/UVEditor/Source/UVEditorTools/Public/Selection/UVToolSelectionAPI.h
@@ -97,6 +97,10 @@ public:
 	 * Gets the bounding box center of the current selection. The center point is cached and
 	 * invalidated on the next SetSelection call, but can be forced to be
 	 * recalculated.
+	 * 
+	 * Note: you should call with bForceRecalculate = true whenever calling this from OnCanonicalModified,
+	 * because it is not guaranteed that the selection api will detect the mesh change (and therefore
+	 * invalidate the cached value) before your call site, due to the order of OnCanonicalModified broadcasts
 	 */
 	FVector3d GetUnwrapSelectionBoundingBoxCenter(bool bForceRecalculate = false);
 
diff --git a/Engine/Plugins/Experimental/UVEditor/Source/UVEditorTools/Public/UVSelectTool.h b/Engine/Plugins/Experimental/UVEditor/Source/UVEditorTools/Public/UVSelectTool.h
--- a/Engine/Plugins/Experimental/UVEditor/Source/UVEditorTools/Public/UVSelectTool.h
+++ b/Engine/Plugins/Experimental/UVEditor/Source/UVEditorTools/Public/UVSelectTool.h
@@ -84,7 +84,7 @@ protected:
 	virtual void GizmoTransformEnded(UTransformProxy* Proxy);
 
 	virtual void ApplyGizmoTransform();
-	virtual void UpdateGizmo();
+	virtual void UpdateGizmo(bool bForceRecomputeSelectionCenters = false);
 
 	UPROPERTY()
 	TArray<TObjectPtr<UUVEditorToolMeshInput>> Targets;
