UVEditor: Provide UI limits for the Transform tool's rotation fields to prevent values reaching points where it causes loss of precision issues in the resulting transformations.

diff --git a/Engine/Plugins/Experimental/UVEditor/Source/UVEditorTools/Public/Operators/UVEditorUVTransformOp.h b/Engine/Plugins/Experimental/UVEditor/Source/UVEditorTools/Public/Operators/UVEditorUVTransformOp.h
--- a/Engine/Plugins/Experimental/UVEditor/Source/UVEditorTools/Public/Operators/UVEditorUVTransformOp.h
+++ b/Engine/Plugins/Experimental/UVEditor/Source/UVEditorTools/Public/Operators/UVEditorUVTransformOp.h
@@ -143,7 +143,8 @@ public:
 	FVector2D Scale = FVector2D(1.0, 1.0); 
 
 	/** Rotation applied to UVs after scaling, specified in degrees */
-	UPROPERTY(EditAnywhere, Category = "Advanced Transform | Rotation", meta = (TransientToolProperty, DisplayName = "Rotation"))
+	UPROPERTY(EditAnywhere, Category = "Advanced Transform | Rotation", meta = (TransientToolProperty, DisplayName = "Rotation", 
+		                                                                        UIMin = "-360", UIMax = "360", ClampMin = "-100000", ClampMax = "100000"))
 	float Rotation = 0.0;
 
 	/** Translation applied to UVs, and after scaling and rotation */
@@ -166,7 +167,8 @@ public:
 	UPROPERTY(EditAnywhere, Category = "Quick Transform", meta = (TransientToolProperty, DisplayName = "Translation"))
 	float QuickTranslateOffset = 0.0;
 
-	UPROPERTY(EditAnywhere, Category = "Quick Transform", meta = (TransientToolProperty, DisplayName = "Rotation"))
+	UPROPERTY(EditAnywhere, Category = "Quick Transform", meta = (TransientToolProperty, DisplayName = "Rotation", 
+		                                                          UIMin = "-360", UIMax = "360", ClampMin = "-100000", ClampMax = "100000"))
 	float QuickRotationOffset = 0.0;
 
 	UPROPERTY(EditAnywhere, Category = "Quick Transform", meta = (TransientToolProperty))
