UVEditor: Prevent crashes when a dynamic mesh is loaded into the UV Editor and then deleted externally before exiting the UV Editor

diff --git a/Engine/Plugins/Experimental/UVEditor/Source/UVEditor/Private/UVEditorToolkit.cpp b/Engine/Plugins/Experimental/UVEditor/Source/UVEditor/Private/UVEditorToolkit.cpp
--- a/Engine/Plugins/Experimental/UVEditor/Source/UVEditor/Private/UVEditorToolkit.cpp
+++ b/Engine/Plugins/Experimental/UVEditor/Source/UVEditor/Private/UVEditorToolkit.cpp
@@ -180,14 +180,21 @@ FText FUVEditorToolkit::GetToolkitToolTipText() const
 	ToolTipString += TEXT(": ");
 
 	const TArray<UObject*>* Objects = GetObjectsCurrentlyBeingEdited();
-	check(Objects && Objects->Num() > 0);
-	ToolTipString += GetLabelForObject((*Objects)[0]).ToString();
-	for (int32 i = 1; i < Objects->Num(); ++i)
+	if (Objects && Objects->Num() > 0)
 	{
-		ToolTipString += TEXT(", ");
-		ToolTipString += GetLabelForObject((*Objects)[i]).ToString();
+		ToolTipString += GetLabelForObject((*Objects)[0]).ToString();
+		for (int32 i = 1; i < Objects->Num(); ++i)
+		{
+			ToolTipString += TEXT(", ");
+			ToolTipString += GetLabelForObject((*Objects)[i]).ToString();
+		}
+	}
+	else
+	{
+		// This can occur if our targets have been deleted externally to the UV Editor.
+		// It's a bad state, but one we can avoid crashing in by doing this.
+		ToolTipString += TEXT("<NO OBJECT>");
 	}
-
 	return FText::FromString(ToolTipString);
 }
 
diff --git a/Engine/Plugins/Experimental/UVEditor/Source/UVEditorTools/Private/Actions/UVSeamSewAction.cpp b/Engine/Plugins/Experimental/UVEditor/Source/UVEditorTools/Private/Actions/UVSeamSewAction.cpp
--- a/Engine/Plugins/Experimental/UVEditor/Source/UVEditorTools/Private/Actions/UVSeamSewAction.cpp
+++ b/Engine/Plugins/Experimental/UVEditor/Source/UVEditorTools/Private/Actions/UVSeamSewAction.cpp
@@ -68,7 +68,7 @@ bool UUVSeamSewAction::ExecuteAction()
 	TArray<FUVToolSelection> StartSelections = SelectionAPI->GetSelections();
 	for (const FUVToolSelection& Selection : StartSelections)
 	{
-		if (!ensure(Selection.Target.IsValid() && Selection.Target->IsValid()))
+		if (!ensure(Selection.Target.IsValid() && Selection.Target->AreMeshesValid()))
 		{
 			continue;
 		}
diff --git a/Engine/Plugins/Experimental/UVEditor/Source/UVEditorTools/Private/Actions/UVSplitAction.cpp b/Engine/Plugins/Experimental/UVEditor/Source/UVEditorTools/Private/Actions/UVSplitAction.cpp
--- a/Engine/Plugins/Experimental/UVEditor/Source/UVEditorTools/Private/Actions/UVSplitAction.cpp
+++ b/Engine/Plugins/Experimental/UVEditor/Source/UVEditorTools/Private/Actions/UVSplitAction.cpp
@@ -24,7 +24,7 @@ namespace UVSplitActionLocals
 	{
 		UUVEditorToolMeshInput* Target = Selection.Target.Get();
 		
-		if (!ensure(Target && Target->IsValid()))
+		if (!ensure(Target && Target->AreMeshesValid()))
 		{
 			return false;
 		}
@@ -107,7 +107,7 @@ namespace UVSplitActionLocals
 	{
 		UUVEditorToolMeshInput* Target = Selection.Target.Get();
 
-		if (!ensure(Target && Target->IsValid()))
+		if (!ensure(Target && Target->AreMeshesValid()))
 		{
 			return false;
 		}
@@ -161,7 +161,7 @@ namespace UVSplitActionLocals
 	{
 		UUVEditorToolMeshInput* Target = Selection.Target.Get();
 
-		if (!ensure(Target && Target->IsValid()))
+		if (!ensure(Target && Target->AreMeshesValid()))
 		{
 			return false;
 		}
diff --git a/Engine/Plugins/Experimental/UVEditor/Source/UVEditorTools/Private/Selection/UVToolSelectionAPI.cpp b/Engine/Plugins/Experimental/UVEditor/Source/UVEditorTools/Private/Selection/UVToolSelectionAPI.cpp
--- a/Engine/Plugins/Experimental/UVEditor/Source/UVEditorTools/Private/Selection/UVToolSelectionAPI.cpp
+++ b/Engine/Plugins/Experimental/UVEditor/Source/UVEditorTools/Private/Selection/UVToolSelectionAPI.cpp
@@ -170,7 +170,7 @@ void UUVToolSelectionAPI::SetSelections(const TArray<FUVToolSelection>& Selectio
 	TArray<FUVToolSelection> NewSelections;
 	for (const FUVToolSelection& NewSelection : SelectionsIn)
 	{
-		if (ensure(NewSelection.Target.IsValid() && NewSelection.Target->IsValid())
+		if (ensure(NewSelection.Target.IsValid() && NewSelection.Target->AreMeshesValid())
 			// All of the selections should match type
 			&& ensure(NewSelections.Num() == 0 || NewSelection.Type == NewSelections[0].Type)
 			// Selection must not be empty, unless it's an edge selection stored as stable identifiers
@@ -315,7 +315,7 @@ void UUVToolSelectionAPI::SetUnsetElementAppliedMeshSelections(const TArray<FUVT
 	TArray<FUVToolSelection> NewUnsetSelections;
 	for (const FUVToolSelection& NewUnsetSelection : UnsetSelectionsIn)
 	{
-		if (ensure(NewUnsetSelection.Target.IsValid() && NewUnsetSelection.Target->IsValid())
+		if (ensure(NewUnsetSelection.Target.IsValid() && NewUnsetSelection.Target->AreMeshesValid())
 			// All of the unset selections should match type
 			&& ensure(NewUnsetSelections.Num() == 0 || NewUnsetSelection.Type == NewUnsetSelections[0].Type)
 			// Unset selection must not be empty
diff --git a/Engine/Plugins/Experimental/UVEditor/Source/UVEditorTools/Private/Selection/UVToolSelectionHighlightMechanic.cpp b/Engine/Plugins/Experimental/UVEditor/Source/UVEditorTools/Private/Selection/UVToolSelectionHighlightMechanic.cpp
--- a/Engine/Plugins/Experimental/UVEditor/Source/UVEditorTools/Private/Selection/UVToolSelectionHighlightMechanic.cpp
+++ b/Engine/Plugins/Experimental/UVEditor/Source/UVEditorTools/Private/Selection/UVToolSelectionHighlightMechanic.cpp
@@ -202,7 +202,7 @@ void UUVToolSelectionHighlightMechanic::RebuildUnwrapHighlight(
 
 	for (const FUVToolSelection& Selection : Selections)
 	{
-		if (!ensure(Selection.Target.IsValid() && Selection.Target->IsValid()))
+		if (!ensure(Selection.Target.IsValid() && Selection.Target->AreMeshesValid()))
 		{
 			return;
 		}
@@ -415,7 +415,7 @@ void UUVToolSelectionHighlightMechanic::RebuildAppliedHighlightFromUnwrapSelecti
 
 	for (const FUVToolSelection& Selection : UnwrapSelections)
 	{
-		if (!ensure(Selection.Target.IsValid() && Selection.Target->IsValid()))
+		if (!ensure(Selection.Target.IsValid() && Selection.Target->AreMeshesValid()))
 		{
 			return;
 		}
diff --git a/Engine/Plugins/Experimental/UVEditor/Source/UVEditorTools/Private/ToolTargets/UVEditorToolMeshInput.cpp b/Engine/Plugins/Experimental/UVEditor/Source/UVEditorTools/Private/ToolTargets/UVEditorToolMeshInput.cpp
--- a/Engine/Plugins/Experimental/UVEditor/Source/UVEditorTools/Private/ToolTargets/UVEditorToolMeshInput.cpp
+++ b/Engine/Plugins/Experimental/UVEditor/Source/UVEditorTools/Private/ToolTargets/UVEditorToolMeshInput.cpp
@@ -61,13 +61,22 @@ namespace UVEditorToolMeshInputLocals
 
 const TArray<int32>* const UUVEditorToolMeshInput::NONE_CHANGED_ARG = &UVEditorToolMeshInputLocals::EmptyArray;
 
-bool UUVEditorToolMeshInput::IsValid() const
+bool UUVEditorToolMeshInput::AreMeshesValid() const
 {
-	return UnwrapCanonical
-		&& UnwrapPreview && UnwrapPreview->IsValidLowLevel()
-		&& AppliedCanonical
+	return UnwrapPreview && UnwrapPreview->IsValidLowLevel()
 		&& AppliedPreview && AppliedPreview->IsValidLowLevel()
-		&& SourceTarget && SourceTarget->IsValid()
+		&& UnwrapCanonical && AppliedCanonical;
+}
+
+bool UUVEditorToolMeshInput::IsToolTargetValid() const
+{
+	return SourceTarget && SourceTarget->IsValid();
+}
+
+bool UUVEditorToolMeshInput::IsValid() const
+{
+	return AreMeshesValid()
+		&& IsToolTargetValid()
 		&& UVLayerIndex >= 0;
 }
 
diff --git a/Engine/Plugins/Experimental/UVEditor/Source/UVEditorTools/Public/ToolTargets/UVEditorToolMeshInput.h b/Engine/Plugins/Experimental/UVEditor/Source/UVEditorTools/Public/ToolTargets/UVEditorToolMeshInput.h
--- a/Engine/Plugins/Experimental/UVEditor/Source/UVEditorTools/Public/ToolTargets/UVEditorToolMeshInput.h
+++ b/Engine/Plugins/Experimental/UVEditor/Source/UVEditorTools/Public/ToolTargets/UVEditorToolMeshInput.h
@@ -245,6 +245,20 @@ public:
 	 */
 	void AppliedVidToUnwrapVids(int32 AppliedVid, TArray<int32>& UnwrapVidsOut);
 
+	/**
+	 * Returns whether or not the underlying source ToolTarget is still valid.
+	 * This is separated from IsValid, which also checks the ToolTarget, in case
+	 * we are interested in the status of just the ToolTarget.
+	*/
+	virtual bool IsToolTargetValid() const;
+
+	/**
+	 * Returns whether or not the canonical and preview meshes are still valid.
+	 * This is separated from IsValid, in case we are interested in the status
+	 * of just the canonical and preview meshes.
+	 */
+	virtual bool AreMeshesValid() const;
+
 	// UToolTarget
 	virtual bool IsValid() const override;
 };
\ No newline at end of file
