UVEditor: Fix issue where selections of edges that are connected by unselected additional edges are treated as fully connected for the Transform/Align/Distribute tools.

diff --git a/Engine/Plugins/Experimental/UVEditor/Source/UVEditorTools/Private/Operators/UVEditorUVTransformOp.cpp b/Engine/Plugins/Experimental/UVEditor/Source/UVEditorTools/Private/Operators/UVEditorUVTransformOp.cpp
--- a/Engine/Plugins/Experimental/UVEditor/Source/UVEditorTools/Private/Operators/UVEditorUVTransformOp.cpp
+++ b/Engine/Plugins/Experimental/UVEditor/Source/UVEditorTools/Private/Operators/UVEditorUVTransformOp.cpp
@@ -32,10 +32,10 @@ void FUVEditorUVTransformBaseOp::SetTransform(const FTransformSRT3d& Transform)
 	ResultTransform = Transform;
 }
 
-void FUVEditorUVTransformBaseOp::SetSelection(const TSet<int32>& TriangleSelectionIn, const TSet<int32>& VertexSelectionIn)
+void FUVEditorUVTransformBaseOp::SetSelection(const TSet<int32>& EdgeSelectionIn, const TSet<int32>& VertexSelectionIn)
 {
-		TriangleSelection = TriangleSelectionIn;
-		VertexSelection = VertexSelectionIn;
+	EdgeSelection = EdgeSelectionIn;
+	VertexSelection = VertexSelectionIn;
 }
 
 void FUVEditorUVTransformBaseOp::SortComponentsByBoundingBox(bool bSortXThenY)
@@ -209,9 +209,16 @@ void FUVEditorUVTransformBaseOp::CalculateResult(FProgressCancel* Progress)
 	ActiveUVLayer = ResultMesh->Attributes()->GetUVLayer(UVLayerInput);
 	ensure(ActiveUVLayer);
 
-	auto UVIslandPredicate = [this](int32 Triangle0, int32 Triangle1)
+	auto UVIslandPredicate = [this](int32 Vertex0, int32 Vertex1)
 	{
-		return GroupingMode != EUVEditorAlignDistributeGroupingModeBackend::IndividualVertices;
+		
+		if (GroupingMode != EUVEditorAlignDistributeGroupingModeBackend::IndividualVertices
+			&& EdgeSelection.IsSet())
+		{
+			const TSet<int32>& SelectedEdges = EdgeSelection.GetValue();
+			return SelectedEdges.Contains(ResultMesh->FindEdge(Vertex0, Vertex1));
+		}
+		return false;
 	};
 
 	UVComponents = MakeShared<FMeshConnectedComponents>(ResultMesh.Get());
@@ -785,9 +792,9 @@ TUniquePtr<FDynamicMeshOperator> UUVEditorUVTransformOperatorFactory::MakeNewOpe
 		TUniquePtr<FUVEditorUVTransformOp> Op = MakeUnique<FUVEditorUVTransformOp>();
 		Op->OriginalMesh = OriginalMesh;
 		Op->SetTransform(TargetTransform);
-		if (TriangleSelection.IsSet() && VertexSelection.IsSet())
+		if (EdgeSelection.IsSet() && VertexSelection.IsSet())
 		{
-			Op->SetSelection(TriangleSelection.GetValue(), VertexSelection.GetValue());
+			Op->SetSelection(EdgeSelection.GetValue(), VertexSelection.GetValue());
 		}
 		Op->UVLayerIndex = GetSelectedUVChannel();
 
@@ -840,9 +847,9 @@ TUniquePtr<FDynamicMeshOperator> UUVEditorUVTransformOperatorFactory::MakeNewOpe
 		TUniquePtr<FUVEditorUVAlignOp> Op = MakeUnique<FUVEditorUVAlignOp>();
 		Op->OriginalMesh = OriginalMesh;
 		Op->SetTransform(TargetTransform);
-		if (TriangleSelection.IsSet() && VertexSelection.IsSet())
+		if (EdgeSelection.IsSet() && VertexSelection.IsSet())
 		{
-			Op->SetSelection(TriangleSelection.GetValue(), VertexSelection.GetValue());
+			Op->SetSelection(EdgeSelection.GetValue(), VertexSelection.GetValue());
 		}
 		Op->UVLayerIndex = GetSelectedUVChannel();
 
@@ -918,9 +925,9 @@ TUniquePtr<FDynamicMeshOperator> UUVEditorUVTransformOperatorFactory::MakeNewOpe
 		TUniquePtr<FUVEditorUVDistributeOp> Op = MakeUnique<FUVEditorUVDistributeOp>();
 		Op->OriginalMesh = OriginalMesh;
 		Op->SetTransform(TargetTransform);
-		if (TriangleSelection.IsSet() && VertexSelection.IsSet())
+		if (EdgeSelection.IsSet() && VertexSelection.IsSet())
 		{
-			Op->SetSelection(TriangleSelection.GetValue(), VertexSelection.GetValue());
+			Op->SetSelection(EdgeSelection.GetValue(), VertexSelection.GetValue());
 		}
 		Op->UVLayerIndex = GetSelectedUVChannel();
 
diff --git a/Engine/Plugins/Experimental/UVEditor/Source/UVEditorTools/Private/UVEditorTransformTool.cpp b/Engine/Plugins/Experimental/UVEditor/Source/UVEditorTools/Private/UVEditorTransformTool.cpp
--- a/Engine/Plugins/Experimental/UVEditor/Source/UVEditorTools/Private/UVEditorTransformTool.cpp
+++ b/Engine/Plugins/Experimental/UVEditor/Source/UVEditorTools/Private/UVEditorTransformTool.cpp
@@ -194,19 +194,25 @@ void UUVEditorTransformTool::Setup()
 		Factory->GetSelectedUVChannel = [&Target]() { return 0; /*Since we're passing in the unwrap mesh, the UV index is always zero.*/ };
 		if (Selection)
 		{
-			// If we have a selection, it's in the unwrapped mesh. We need both triangles, which are 1:1 between them,
-			// and vertices, which are not, in the applied mesh to pass to the factory.
+			// Generate vertex and edge selection sets for the operation. These are needed to differentiate
+			// between islands, whether they are composed of triangles, edges or isolated points.
 			FUVToolSelection UnwrapVertexSelection;
-			if (Selection->Type == FUVToolSelection::EType::Vertex)
+			FUVToolSelection UnwrapEdgeSelection;
+			switch (Selection->Type)
 			{
-				UnwrapVertexSelection = *Selection;
+			case FUVToolSelection::EType::Vertex:
+				Factory->VertexSelection.Emplace(Selection->SelectedIDs);
+				Factory->EdgeSelection.Emplace(Selection->GetConvertedSelection(*Target.UnwrapCanonical, FUVToolSelection::EType::Edge).SelectedIDs);
+				break;
+			case FUVToolSelection::EType::Edge:
+				Factory->VertexSelection.Emplace(Selection->GetConvertedSelection(*Target.UnwrapCanonical, FUVToolSelection::EType::Vertex).SelectedIDs);
+				Factory->EdgeSelection.Emplace(Selection->SelectedIDs);
+				break;
+			case FUVToolSelection::EType::Triangle:
+				Factory->VertexSelection.Emplace(Selection->GetConvertedSelection(*Target.UnwrapCanonical, FUVToolSelection::EType::Vertex).SelectedIDs);
+				Factory->EdgeSelection.Emplace(Selection->GetConvertedSelection(*Target.UnwrapCanonical, FUVToolSelection::EType::Edge).SelectedIDs);
+				break;
 			}
-			else
-			{
-				UnwrapVertexSelection = Selection->GetConvertedSelection(*Target.UnwrapCanonical, FUVToolSelection::EType::Vertex);
-			}				
-			Factory->VertexSelection.Emplace(UnwrapVertexSelection.SelectedIDs);
-			Factory->TriangleSelection.Emplace(UnwrapVertexSelection.GetConvertedSelection(*Target.UnwrapCanonical, FUVToolSelection::EType::Triangle).SelectedIDs);
 		}
 
 		Target.UnwrapPreview->ChangeOpFactory(Factory);
diff --git a/Engine/Plugins/Experimental/UVEditor/Source/UVEditorTools/Public/Operators/UVEditorUVTransformOp.h b/Engine/Plugins/Experimental/UVEditor/Source/UVEditorTools/Public/Operators/UVEditorUVTransformOp.h
--- a/Engine/Plugins/Experimental/UVEditor/Source/UVEditorTools/Public/Operators/UVEditorUVTransformOp.h
+++ b/Engine/Plugins/Experimental/UVEditor/Source/UVEditorTools/Public/Operators/UVEditorUVTransformOp.h
@@ -338,7 +338,7 @@ public:
 	EUVEditorAlignDistributeGroupingModeBackend GroupingMode;
 
 	void SetTransform(const FTransformSRT3d& Transform);
-	void SetSelection(const TSet<int32>& TriangleSelection, const TSet<int32>& VertexSelection);
+	void SetSelection(const TSet<int32>& EdgeSelection, const TSet<int32>& VertexSelection);
 
 	virtual void CalculateResult(FProgressCancel* Progress) override;
 	const TArray<FVector2D>& GetPivotLocations() const { return VisualizationPivots; }
@@ -361,8 +361,8 @@ protected:
 	TArray<int32> SpatiallyOrderedComponentIndex;
 	TArray<FVector2D> VisualizationPivots;
 
-	TOptional<TSet<int32>> TriangleSelection;
 	TOptional<TSet<int32>> VertexSelection;
+	TOptional<TSet<int32>> EdgeSelection;
 };
 
 
@@ -445,7 +445,7 @@ public:
 	UPROPERTY()
 	TObjectPtr<UUVEditorUVTransformPropertiesBase> Settings;
 
-	TOptional<TSet<int32>> TriangleSelection;
+	TOptional<TSet<int32>> EdgeSelection;
 	TOptional<TSet<int32>> VertexSelection;
 
 	TSharedPtr<UE::Geometry::FDynamicMesh3, ESPMode::ThreadSafe> OriginalMesh;
