UVEditor: Fix condition on TransformBaseOP for computing connected components to correctly handle the different grouping modes and selection statuses. This corrects an issue with the previous bug fix that incorrectly treated islands as disconnected vertices.

diff --git a/Engine/Plugins/Experimental/UVEditor/Source/UVEditorTools/Private/Operators/UVEditorUVTransformOp.cpp b/Engine/Plugins/Experimental/UVEditor/Source/UVEditorTools/Private/Operators/UVEditorUVTransformOp.cpp
--- a/Engine/Plugins/Experimental/UVEditor/Source/UVEditorTools/Private/Operators/UVEditorUVTransformOp.cpp
+++ b/Engine/Plugins/Experimental/UVEditor/Source/UVEditorTools/Private/Operators/UVEditorUVTransformOp.cpp
@@ -211,14 +211,21 @@ void FUVEditorUVTransformBaseOp::CalculateResult(FProgressCancel* Progress)
 
 	auto UVIslandPredicate = [this](int32 Vertex0, int32 Vertex1)
 	{
-		
+		if (GroupingMode == EUVEditorAlignDistributeGroupingModeBackend::IndividualVertices)
+		{
+			return false;
+		}
+
 		if (GroupingMode != EUVEditorAlignDistributeGroupingModeBackend::IndividualVertices
 			&& EdgeSelection.IsSet())
 		{
 			const TSet<int32>& SelectedEdges = EdgeSelection.GetValue();
 			return SelectedEdges.Contains(ResultMesh->FindEdge(Vertex0, Vertex1));
 		}
-		return false;
+		else
+		{
+			return true;
+		}
 	};
 
 	UVComponents = MakeShared<FMeshConnectedComponents>(ResultMesh.Get());
