UVEditor: Fix issue where the UX was in a semi-broken state after an autosave event occurs within the Editor at large.

diff --git a/Engine/Plugins/Experimental/UVEditor/Source/UVEditor/Private/UVEditorMode.cpp b/Engine/Plugins/Experimental/UVEditor/Source/UVEditor/Private/UVEditorMode.cpp
--- a/Engine/Plugins/Experimental/UVEditor/Source/UVEditor/Private/UVEditorMode.cpp
+++ b/Engine/Plugins/Experimental/UVEditor/Source/UVEditor/Private/UVEditorMode.cpp
@@ -49,6 +49,7 @@
 #include "EditorSupportDelegates.h"
 #include "Utilities/MeshUDIMClassifier.h"
 #include "UVEditorLogging.h"
+#include "UObject/ObjectSaveContext.h"
 
 #include UE_INLINE_GENERATED_CPP_BY_NAME(UVEditorMode)
 
@@ -333,6 +334,11 @@ void UUVEditorMode::Enter()
 		SetSimulationWarning(false);
 	});
 
+	PostSaveWorldDelegateHandle = FEditorDelegates::PostSaveWorldWithContext.AddLambda([this](UWorld*, FObjectPostSaveContext)
+	{
+		ActivateDefaultTool();
+	});
+
 	// Our mode needs to get Render and DrawHUD calls, but doesn't implement the legacy interface that
 	// causes those to be called. Instead, we'll hook into the tools context's Render and DrawHUD calls.
 	GetInteractiveToolsContext()->OnRender.AddUObject(this, &UUVEditorMode::Render);
@@ -859,6 +865,7 @@ void UUVEditorMode::Exit()
 	FEditorDelegates::PreBeginPIE.Remove(BeginPIEDelegateHandle);
 	FEditorDelegates::EndPIE.Remove(EndPIEDelegateHandle);
 	FEditorDelegates::CancelPIE.Remove(CancelPIEDelegateHandle);
+	FEditorDelegates::PostSaveWorldWithContext.Remove(PostSaveWorldDelegateHandle);
 
 	Super::Exit();
 }
@@ -1367,6 +1374,12 @@ void UUVEditorMode::ModeTick(float DeltaTime)
 		ToolInput->UnwrapPreview->Tick(DeltaTime);
 	}
 
+	// If nothing is running at this point, restart the default tool, since it needs to be running to handle some toolbar UX
+	if (GetInteractiveToolsContext() && !GetInteractiveToolsContext()->HasActiveTool())
+	{
+		ActivateDefaultTool();
+	}
+
 }
 
 void UUVEditorMode::SetSimulationWarning(bool bEnabled)
diff --git a/Engine/Plugins/Experimental/UVEditor/Source/UVEditor/Public/UVEditorMode.h b/Engine/Plugins/Experimental/UVEditor/Source/UVEditor/Public/UVEditorMode.h
--- a/Engine/Plugins/Experimental/UVEditor/Source/UVEditor/Public/UVEditorMode.h
+++ b/Engine/Plugins/Experimental/UVEditor/Source/UVEditor/Public/UVEditorMode.h
@@ -371,5 +371,9 @@ protected:
 	FDelegateHandle BeginPIEDelegateHandle;
 	FDelegateHandle EndPIEDelegateHandle;
 	FDelegateHandle CancelPIEDelegateHandle;
+
+	// Holds references to Save callbacks to handle logic when the autosave triggers and shuts down active tools. 
+	// We need to recover from this, so we restart the select tool after the save is over.
+	FDelegateHandle PostSaveWorldDelegateHandle;
 };
 
