UVEditor: Prevent crash upon opening the UV Editor if the live preview tab had been closed prior to last exiting the UV Editor. Also forces the live preview tab to be opened automatically upon UV Editor startup.

diff --git a/Engine/Plugins/Experimental/UVEditor/Source/UVEditor/Private/UVEditorMode.cpp b/Engine/Plugins/Experimental/UVEditor/Source/UVEditor/Private/UVEditorMode.cpp
--- a/Engine/Plugins/Experimental/UVEditor/Source/UVEditor/Private/UVEditorMode.cpp
+++ b/Engine/Plugins/Experimental/UVEditor/Source/UVEditor/Private/UVEditorMode.cpp
@@ -896,7 +896,14 @@ void UUVEditorMode::InitializeAssetEditorContexts(UContextObjectStore& ContextSt
 				GetCameraState(*LivePreviewViewportClientPtr, CameraStateOut);
 			},
 			[LivePreviewViewportClientPtr = &LivePreviewViewportClient](const FAxisAlignedBox3d& BoundingBox) {
-				LivePreviewViewportClientPtr->FocusViewportOnBox((FBox)BoundingBox, true);
+				// We check for the Viewport here because it might not be open at the time this
+				// method is called, e.g. during startup with an initially closed tab. And since
+				// the FocusViewportOnBox method doesn't check internally that the Viewport is
+				// available, this can crash.
+				if (LivePreviewViewportClientPtr && LivePreviewViewportClientPtr->Viewport)
+				{
+					LivePreviewViewportClientPtr->FocusViewportOnBox((FBox)BoundingBox, true);
+				}
 			}
 			);
 		ContextStore.AddContextObject(LivePreviewAPI);
diff --git a/Engine/Plugins/Experimental/UVEditor/Source/UVEditor/Private/UVEditorToolkit.cpp b/Engine/Plugins/Experimental/UVEditor/Source/UVEditor/Private/UVEditorToolkit.cpp
--- a/Engine/Plugins/Experimental/UVEditor/Source/UVEditor/Private/UVEditorToolkit.cpp
+++ b/Engine/Plugins/Experimental/UVEditor/Source/UVEditor/Private/UVEditorToolkit.cpp
@@ -457,6 +457,16 @@ void FUVEditorToolkit::PostInitAssetEditor()
 		TabManager->TryInvokeTab(ViewportTabID);
 	}
 
+    // Note: We don't have to force the live viewport to be open, but if we don't, we need to
+	// make sure that any future live preview api functionality does not crash when the viewport
+	// is missing, because some viewport client functions are not robust to that case.For now,
+	// we'll just force it open because it is safer and seems to be more convenient for the user,
+	// since reopening the window can be unintuitive, whereas closing it is easy.
+	if (!TabManager->FindExistingLiveTab(LivePreviewTabID))
+	{
+		TabManager->TryInvokeTab(LivePreviewTabID);
+	}
+
 	// Add the "Apply Changes" button. It should actually be safe to do this almost
 	// any time, even before that toolbar's registration, but it's easier to put most
 	// things into PostInitAssetEditor().
