UVEditor: Fix crash where right clicking on the asset tab in the editor would fail an assertion.

Moved incorrect placement of the mode deactivation into the OnClose() method of the toolkit rather than the OnRequestClose(). Further, adds more tolerant ensures instead of hard checks around the mode's existence in the crashing routines.

diff --git a/Engine/Plugins/Experimental/UVEditor/Source/UVEditor/Private/UVEditorToolkit.cpp b/Engine/Plugins/Experimental/UVEditor/Source/UVEditor/Private/UVEditorToolkit.cpp
--- a/Engine/Plugins/Experimental/UVEditor/Source/UVEditor/Private/UVEditorToolkit.cpp
+++ b/Engine/Plugins/Experimental/UVEditor/Source/UVEditor/Private/UVEditorToolkit.cpp
@@ -282,13 +282,18 @@ bool FUVEditorToolkit::OnRequestClose()
 		}
 	}
 
+	return FAssetEditorToolkit::OnRequestClose();
+}
+
+void FUVEditorToolkit::OnClose()
+{
 	// Give any active modes a chance to shutdown while the toolkit host is still alive
-    // This is super important to do, otherwise currently opened tabs won't be marked as "closed".
-    // This results in tabs not being properly recycled upon reopening the editor and tab
-    // duplication for each opening event.
+	// This is super important to do, otherwise currently opened tabs won't be marked as "closed".
+	// This results in tabs not being properly recycled upon reopening the editor and tab
+	// duplication for each opening event.
 	GetEditorModeManager().ActivateDefaultMode();
 
-	return FAssetEditorToolkit::OnRequestClose();
+	FAssetEditorToolkit::OnClose();
 }
 
 // These get called indirectly (via toolkit host) from the mode toolkit when the mode starts or ends a tool,
@@ -308,8 +313,7 @@ void FUVEditorToolkit::RemoveViewportOverlayWidget(TSharedRef<SWidget> InViewpor
 void FUVEditorToolkit::SaveAsset_Execute()
 {
 	UUVEditorMode* UVMode = Cast<UUVEditorMode>(EditorModeManager->GetActiveScriptableMode(UUVEditorMode::EM_UVEditorModeId));
-	check(UVMode);
-	if (UVMode->HaveUnappliedChanges())
+	if (ensure(UVMode) && UVMode->HaveUnappliedChanges())
 	{
 		UVMode->ApplyChanges();
 	}
@@ -320,8 +324,11 @@ void FUVEditorToolkit::SaveAsset_Execute()
 bool FUVEditorToolkit::CanSaveAsset() const 
 {
 	UUVEditorMode* UVMode = Cast<UUVEditorMode>(EditorModeManager->GetActiveScriptableMode(UUVEditorMode::EM_UVEditorModeId));
-	check(UVMode);
-	return UVMode->CanApplyChanges();
+	if(ensure(UVMode))	
+	{
+		return UVMode->CanApplyChanges();
+	}
+	return false;	
 }
 
 bool FUVEditorToolkit::CanSaveAssetAs() const 
diff --git a/Engine/Plugins/Experimental/UVEditor/Source/UVEditor/Public/UVEditorToolkit.h b/Engine/Plugins/Experimental/UVEditor/Source/UVEditor/Public/UVEditorToolkit.h
--- a/Engine/Plugins/Experimental/UVEditor/Source/UVEditor/Public/UVEditorToolkit.h
+++ b/Engine/Plugins/Experimental/UVEditor/Source/UVEditor/Public/UVEditorToolkit.h
@@ -50,6 +50,7 @@ public:
 	virtual FText GetToolkitToolTipText() const override;
 	virtual void RegisterTabSpawners(const TSharedRef<FTabManager>& InTabManager) override;
 	virtual bool OnRequestClose() override;
+	virtual void OnClose() override;
 	virtual void SaveAsset_Execute() override;
 	virtual bool CanSaveAsset() const override;
 	virtual bool CanSaveAssetAs() const override;
