UVEditor: Moved analytics event name creation out of public header

diff --git a/Engine/Plugins/Experimental/UVEditor/Source/UVEditorTools/Private/UVEditorLayerEditTool.cpp b/Engine/Plugins/Experimental/UVEditor/Source/UVEditorTools/Private/UVEditorLayerEditTool.cpp
--- a/Engine/Plugins/Experimental/UVEditor/Source/UVEditorTools/Private/UVEditorLayerEditTool.cpp
+++ b/Engine/Plugins/Experimental/UVEditor/Source/UVEditorTools/Private/UVEditorLayerEditTool.cpp
@@ -752,7 +752,7 @@ void UUVEditorChannelEditTool::RecordAnalytics()
 	{
 		for (const FAnalyticsEventAttribute& Attr : Attributes)
 		{
-			UE_LOG(LogGeometry, Log, TEXT("Debug %s.ChannelEditTool.%s = %s"), *UVEditorAnalyticsPrefix, *Attr.GetName(), *Attr.GetValue());
+			UE_LOG(LogGeometry, Log, TEXT("Debug %s.%s = %s"), *UVEditorAnalyticsEventName(TEXT("ChannelEditTool")), *Attr.GetName(), *Attr.GetValue());
 		}
 	}
 }
diff --git a/Engine/Plugins/Experimental/UVEditor/Source/UVEditorTools/Private/UVEditorLayoutTool.cpp b/Engine/Plugins/Experimental/UVEditor/Source/UVEditorTools/Private/UVEditorLayoutTool.cpp
--- a/Engine/Plugins/Experimental/UVEditor/Source/UVEditorTools/Private/UVEditorLayoutTool.cpp
+++ b/Engine/Plugins/Experimental/UVEditor/Source/UVEditorTools/Private/UVEditorLayoutTool.cpp
@@ -244,7 +244,7 @@ void UUVEditorLayoutTool::RecordAnalytics()
 	{
 		for (const FAnalyticsEventAttribute& Attr : Attributes)
 		{
-			UE_LOG(LogGeometry, Log, TEXT("Debug %s.LayoutTool.%s = %s"), *UVEditorAnalyticsPrefix, *Attr.GetName(), *Attr.GetValue());
+			UE_LOG(LogGeometry, Log, TEXT("Debug %s.%s = %s"), *UVEditorAnalyticsEventName(TEXT("LayoutTool")), *Attr.GetName(), *Attr.GetValue());
 		}
 	}
 }
diff --git a/Engine/Plugins/Experimental/UVEditor/Source/UVEditorTools/Private/UVEditorRecomputeUVsTool.cpp b/Engine/Plugins/Experimental/UVEditor/Source/UVEditorTools/Private/UVEditorRecomputeUVsTool.cpp
--- a/Engine/Plugins/Experimental/UVEditor/Source/UVEditorTools/Private/UVEditorRecomputeUVsTool.cpp
+++ b/Engine/Plugins/Experimental/UVEditor/Source/UVEditorTools/Private/UVEditorRecomputeUVsTool.cpp
@@ -338,7 +338,7 @@ void UUVEditorRecomputeUVsTool::RecordAnalytics()
 	{
 		for (const FAnalyticsEventAttribute& Attr : Attributes)
 		{
-			UE_LOG(LogGeometry, Log, TEXT("Debug %s.UnwrapTool.%s = %s"), *UVEditorAnalyticsPrefix, *Attr.GetName(), *Attr.GetValue());
+			UE_LOG(LogGeometry, Log, TEXT("Debug %s.%s = %s"), *UVEditorAnalyticsEventName(TEXT("UnwrapTool")), *Attr.GetName(), *Attr.GetValue());
 		}
 	}
 }
diff --git a/Engine/Plugins/Experimental/UVEditor/Source/UVEditorTools/Private/UVEditorSeamTool.cpp b/Engine/Plugins/Experimental/UVEditor/Source/UVEditorTools/Private/UVEditorSeamTool.cpp
--- a/Engine/Plugins/Experimental/UVEditor/Source/UVEditorTools/Private/UVEditorSeamTool.cpp
+++ b/Engine/Plugins/Experimental/UVEditor/Source/UVEditorTools/Private/UVEditorSeamTool.cpp
@@ -1215,7 +1215,7 @@ void UUVEditorSeamTool::RecordAnalytics()
 	{
 		for (const FAnalyticsEventAttribute& Attr : Attributes)
 		{
-			UE_LOG(LogGeometry, Log, TEXT("Debug %s.SeamTool.%s = %s"), *UVEditorAnalyticsPrefix, *Attr.GetName(), *Attr.GetValue());
+			UE_LOG(LogGeometry, Log, TEXT("Debug %s.%s = %s"), *UVEditorAnalyticsEventName(TEXT("SeamTool")), *Attr.GetName(), *Attr.GetValue());
 		}
 	}
 }
diff --git a/Engine/Plugins/Experimental/UVEditor/Source/UVEditorTools/Private/UVEditorToolAnalyticsUtils.cpp b/Engine/Plugins/Experimental/UVEditor/Source/UVEditorTools/Private/UVEditorToolAnalyticsUtils.cpp
--- a/Engine/Plugins/Experimental/UVEditor/Source/UVEditorTools/Private/UVEditorToolAnalyticsUtils.cpp
+++ b/Engine/Plugins/Experimental/UVEditor/Source/UVEditorTools/Private/UVEditorToolAnalyticsUtils.cpp
@@ -13,6 +13,11 @@ namespace Geometry
 namespace UVEditorAnalytics
 {
 
+FString UVEditorAnalyticsEventName(const FString& ToolEventName)
+{
+	return FString::Printf(TEXT("Editor.Usage.UVEditor.%s"), *ToolEventName);
+}
+
 void FTargetAnalytics::AppendToAttributes(TArray<FAnalyticsEventAttribute>& Attributes, FString Prefix) const
 {
 	if (!Prefix.IsEmpty())
diff --git a/Engine/Plugins/Experimental/UVEditor/Source/UVEditorTools/Private/UVEditorTransformTool.cpp b/Engine/Plugins/Experimental/UVEditor/Source/UVEditorTools/Private/UVEditorTransformTool.cpp
--- a/Engine/Plugins/Experimental/UVEditor/Source/UVEditorTools/Private/UVEditorTransformTool.cpp
+++ b/Engine/Plugins/Experimental/UVEditor/Source/UVEditorTools/Private/UVEditorTransformTool.cpp
@@ -447,7 +447,7 @@ void UUVEditorTransformTool::RecordAnalytics()
 	{
 		for (const FAnalyticsEventAttribute& Attr : Attributes)
 		{
-			UE_LOG(LogGeometry, Log, TEXT("Debug %s.TransformTool.%s = %s"), *UVEditorAnalyticsPrefix, *Attr.GetName(), *Attr.GetValue());
+			UE_LOG(LogGeometry, Log, TEXT("Debug %s.%s = %s"), *UVEditorAnalyticsEventName(TEXT("TransformTool")), *Attr.GetName(), *Attr.GetValue());
 		}
 	}
 }
diff --git a/Engine/Plugins/Experimental/UVEditor/Source/UVEditorTools/Public/UVEditorToolAnalyticsUtils.h b/Engine/Plugins/Experimental/UVEditor/Source/UVEditorTools/Public/UVEditorToolAnalyticsUtils.h
--- a/Engine/Plugins/Experimental/UVEditor/Source/UVEditorTools/Public/UVEditorToolAnalyticsUtils.h
+++ b/Engine/Plugins/Experimental/UVEditor/Source/UVEditorTools/Public/UVEditorToolAnalyticsUtils.h
@@ -22,12 +22,7 @@ FAnalyticsEventAttribute AnalyticsEventAttributeEnum(const FString& AttributeNam
 	return FAnalyticsEventAttribute(AttributeName, EnumValueName);
 }
 
-static const FString UVEditorAnalyticsPrefix(TEXT("Editor.Usage.UVEditor"));
-
-FORCEINLINE FString UVEditorAnalyticsEventName(const FString& ToolEventName)
-{
-	return FString::Printf(TEXT("%s.%s"), *UVEditorAnalyticsPrefix, *ToolEventName);
-}
+UVEDITORTOOLS_API FString UVEditorAnalyticsEventName(const FString& ToolEventName);
 
 struct UVEDITORTOOLS_API FTargetAnalytics
 {
diff --git a/Engine/Plugins/Experimental/UVEditor/Source/UVEditorToolsEditorOnly/Private/UVEditorParameterizeMeshTool.cpp b/Engine/Plugins/Experimental/UVEditor/Source/UVEditorToolsEditorOnly/Private/UVEditorParameterizeMeshTool.cpp
--- a/Engine/Plugins/Experimental/UVEditor/Source/UVEditorToolsEditorOnly/Private/UVEditorParameterizeMeshTool.cpp
+++ b/Engine/Plugins/Experimental/UVEditor/Source/UVEditorToolsEditorOnly/Private/UVEditorParameterizeMeshTool.cpp
@@ -331,7 +331,7 @@ void UUVEditorParameterizeMeshTool::RecordAnalytics()
 	{
 		for (const FAnalyticsEventAttribute& Attr : Attributes)
 		{
-			UE_LOG(LogGeometry, Log, TEXT("Debug %s.AutoUVTool.%s = %s"), *UVEditorAnalyticsPrefix, *Attr.GetName(), *Attr.GetValue());
+			UE_LOG(LogGeometry, Log, TEXT("Debug %s.%s = %s"), *UVEditorAnalyticsEventName(TEXT("AutoUVTool")), *Attr.GetName(), *Attr.GetValue());
 		}
 	}
 }
