UVEditor: Implement functionality to allow the Seam tool to now remove seams, as well as add them.

* Provides a mode selection for the Seam tool, switching between seam cutting mode and seam joining mode.
* Seam joining mode has an optional behavioral flag to have the path finding preferentially follow existing seams, making it easier to trace seam edges.
* DynamicMeshOverlay now supports merging elements, providing a method to merge one element to another, if they share the same parent, collapsing any references from the first element into references to the second.
* DynamicMeshUVEditor now supports removing seams from edges, providing a method to remove (via the new DynamicMeshOverlay functionality) seams from an edge collection by iteratively merging elements together.

diff --git a/Engine/Plugins/Experimental/UVEditor/Source/UVEditor/Private/UVEditorCommands.cpp b/Engine/Plugins/Experimental/UVEditor/Source/UVEditor/Private/UVEditorCommands.cpp
--- a/Engine/Plugins/Experimental/UVEditor/Source/UVEditor/Private/UVEditorCommands.cpp
+++ b/Engine/Plugins/Experimental/UVEditor/Source/UVEditor/Private/UVEditorCommands.cpp
@@ -30,7 +30,7 @@ void FUVEditorCommands::RegisterCommands()
 	UI_COMMAND(BeginDistributeTool, "Distribute", "Distribute existing UVs", EUserInterfaceActionType::ToggleButton, FInputChord());
 	UI_COMMAND(BeginParameterizeMeshTool, "AutoUV", "Auto-unwrap and pack UVs", EUserInterfaceActionType::ToggleButton, FInputChord());
 	UI_COMMAND(BeginChannelEditTool, "Channels", "Modify UV channels", EUserInterfaceActionType::ToggleButton, FInputChord());
-	UI_COMMAND(BeginSeamTool, "Seam", "Add UV seams", EUserInterfaceActionType::ToggleButton, FInputChord());
+	UI_COMMAND(BeginSeamTool, "Seam", "Edit UV seams", EUserInterfaceActionType::ToggleButton, FInputChord());
 	UI_COMMAND(BeginRecomputeUVsTool, "Unwrap", "Perform UV unwrapping", EUserInterfaceActionType::ToggleButton, FInputChord());
 
 	// These currently get linked to actions inside the select tool, but will eventually have their own buttons among the tools
diff --git a/Engine/Plugins/Experimental/UVEditor/Source/UVEditorTools/Private/UVEditorSeamTool.cpp b/Engine/Plugins/Experimental/UVEditor/Source/UVEditorTools/Private/UVEditorSeamTool.cpp
--- a/Engine/Plugins/Experimental/UVEditor/Source/UVEditorTools/Private/UVEditorSeamTool.cpp
+++ b/Engine/Plugins/Experimental/UVEditor/Source/UVEditorTools/Private/UVEditorSeamTool.cpp
@@ -18,6 +18,7 @@
 #include "UVEditorUXSettings.h"
 #include "ContextObjects/UVToolContextObjects.h"
 #include "EngineAnalytics.h"
+#include "Algo/Unique.h"
 
 #include UE_INLINE_GENERATED_CPP_BY_NAME(UVEditorSeamTool)
 
@@ -40,7 +41,48 @@ namespace UVEditorSeamToolLocals
 	
 	const FString& ExistingSeamsID(TEXT("SeamLineSet"));
 
-	void GetVidPath(FDynamicMesh3* Mesh, const TArray<int32>& StartVids, int32 EndVid, TArray<int32>& VidPathOut)
+	constexpr double SeamEuclideanDistanceMultiplier = 0.01;
+
+	typedef TFunction<double(int32 FromVID, int32 ToVID, int32 SeedVID, double EuclideanDistance)> FPathDistanceFunc;
+
+	FPathDistanceFunc GetUVDistanceFunction(const UUVEditorToolMeshInput& Target, bool bUnwrapMeshSource, bool bTrackBoundaryEdges)
+	{
+		if (!bTrackBoundaryEdges)
+		{
+			return [&Target](int32 FromVID, int32 ToVID, int32 SeedVID, double EuclideanDistance)
+			{
+				return EuclideanDistance;
+			};
+		}
+		else
+		{
+			if (bUnwrapMeshSource)
+			{
+				return [&Target](int32 FromVID, int32 ToVID, int32 SeedVID, double EuclideanDistance)
+				{
+					if (Target.UnwrapCanonical.Get()->IsBoundaryEdge(Target.UnwrapCanonical.Get()->FindEdge(FromVID, ToVID)))
+					{
+						return EuclideanDistance * SeamEuclideanDistanceMultiplier;
+					}
+					return EuclideanDistance;
+				};
+			}
+			else
+			{
+				return [&Target](int32 FromVID, int32 ToVID, int32 SeedVID, double EuclideanDistance)
+				{					
+					const FDynamicMeshUVOverlay* CanonicalOverlay = Target.AppliedCanonical->Attributes()->GetUVLayer(Target.UVLayerIndex);
+					if (CanonicalOverlay->IsSeamEdge(Target.AppliedCanonical->FindEdge(FromVID, ToVID)))
+					{
+						return EuclideanDistance * SeamEuclideanDistanceMultiplier;
+					}
+					return EuclideanDistance;
+				};
+			}
+		}
+	}
+
+	void GetVidPath(FDynamicMesh3* Mesh, const TArray<int32>& StartVids, int32 EndVid, TArray<int32>& VidPathOut, FPathDistanceFunc PathDistanceFunc = nullptr)
 	{
 		VidPathOut.Reset();
 		if (StartVids.Contains(EndVid))
@@ -55,7 +97,13 @@ namespace UVEditorSeamToolLocals
 		{
 			SeedPoints.Add({ StartVid, StartVid, 0 });
 		}
-		
+
+		if (PathDistanceFunc)
+		{
+			PathFinder.bEnableDistanceWeighting = true;
+			PathFinder.GetWeightedDistanceFunc = PathDistanceFunc;
+		}
+
 		if (PathFinder.ComputeToTargetPoint(SeedPoints, EndVid))
 		{
 			PathFinder.FindPathToNearestSeed(EndVid, VidPathOut);
@@ -64,9 +112,9 @@ namespace UVEditorSeamToolLocals
 	}
 
 	// Like the other GetVidPath, just starts from a specific vert rather than from multiple
-	void GetVidPath(FDynamicMesh3* Mesh, int32 StartVid, int32 EndVid, TArray<int32>& VidPathOut)
+	void GetVidPath(FDynamicMesh3* Mesh, int32 StartVid, int32 EndVid, TArray<int32>& VidPathOut, FPathDistanceFunc PathDistanceFunc = nullptr)
 	{
-		GetVidPath(Mesh, TArray<int32> {StartVid}, EndVid, VidPathOut);
+		GetVidPath(Mesh, TArray<int32> {StartVid}, EndVid, VidPathOut, PathDistanceFunc);
 	}
 
 	void AddDisplayedPoints(UUVEditorToolMeshInput* InputObject,
@@ -87,7 +135,7 @@ namespace UVEditorSeamToolLocals
 
 	void AddDisplayedPath(UUVEditorToolMeshInput* InputObject, 
 		ULineSetComponent* UnwrapLineSet, ULineSetComponent* AppliedLineSet,
-		const TArray<int32>& NewPathVids, bool bPathIsFromUnwrap, 
+		const TArray<int32>& NewPathVids, bool bPathIsFromUnwrap, bool bForceDisplaySeamPaths,
 		const FColor& Color, float Thickness, float DepthBias)
 	{
 		TRACE_CPUPROFILER_EVENT_SCOPE(AddDisplayedPath);
@@ -144,11 +192,11 @@ namespace UVEditorSeamToolLocals
 		}
 		else
 		{
-			// We'll only add an edge to the unwrap path if its unique. If it's not, the edge
-			// must be on a seam, and it's not useful to visualize it.
-			auto GetUniqueUnwrapEid = [InputObject](const TArray<int32>& Vids1, const TArray<int32>& Vids2)
+			// Since the tool supports both cutting and joining seams, we want to display all matching edges, so we
+			// can highlight seams in the unwrap mesh when selecting on the applied mesh.
+			auto GetUnwrapEids = [InputObject](const TArray<int32>& Vids1, const TArray<int32>& Vids2)
 			{
-				int32 UniqueEid = IndexConstants::InvalidID;
+				TArray<int32> Eids;
 				for (int32 Vid1 : Vids1)
 				{
 					for (int32 Vid2 : Vids2)
@@ -156,20 +204,13 @@ namespace UVEditorSeamToolLocals
 						int32 Eid = InputObject->UnwrapCanonical->FindEdge(Vid1, Vid2);
 						if (Eid != IndexConstants::InvalidID)
 						{
-							// Make sure this is the first edge we've found.
-							if (UniqueEid == IndexConstants::InvalidID)
-							{
-								UniqueEid = Eid;
-							}
-							else
-							{
-								// Not unique
-								return IndexConstants::InvalidID;
-							}
+							Eids.Add(Eid);
 						}
 					}
 				}
-				return UniqueEid;
+				Eids.Sort();
+				Eids.SetNum(Algo::Unique(Eids));
+				return Eids;
 			};
 
 			TArray<int32> Vids1, Vids2;
@@ -180,14 +221,21 @@ namespace UVEditorSeamToolLocals
 			for (int32 i = 1; i < NewPathVids.Num(); ++i)
 			{
 				InputObject->AppliedVidToUnwrapVids(NewPathVids[i], *EndPoints2);
-				int32 UnwrapEid = GetUniqueUnwrapEid(*EndPoints1, *EndPoints2);
-				if (UnwrapEid != IndexConstants::InvalidID)
+				TArray<int32> UnwrapEids = GetUnwrapEids(*EndPoints1, *EndPoints2);
+
+				// Normally we want to just want to render unwrap edges if there's a one to one relationship
+				// between the applied mesh, but in other cases, like when doing seam joining, we want to
+				// render all replicas to show where joins will happen.
+				if (bForceDisplaySeamPaths || UnwrapEids.Num() == 1)
 				{
-					FIndex2i EdgeVids = InputObject->UnwrapCanonical->GetEdgeV(UnwrapEid);
-					UnwrapLineSet->AddLine(FRenderableLine(
-						(FVector)InputObject->UnwrapCanonical->GetVertex(EdgeVids.A), 
-						(FVector)InputObject->UnwrapCanonical->GetVertex(EdgeVids.B),
-						Color, Thickness, DepthBias));
+					for (int32 UnwrapEid : UnwrapEids)
+					{
+						FIndex2i EdgeVids = InputObject->UnwrapCanonical->GetEdgeV(UnwrapEid);
+						UnwrapLineSet->AddLine(FRenderableLine(
+							(FVector)InputObject->UnwrapCanonical->GetVertex(EdgeVids.A),
+							(FVector)InputObject->UnwrapCanonical->GetVertex(EdgeVids.B),
+							Color, Thickness, DepthBias));
+					}
 				}
 				Swap(EndPoints1, EndPoints2);
 			}
@@ -296,6 +344,11 @@ void UUVEditorSeamTool::Setup()
 
 	UInteractiveTool::Setup();
 
+	Settings = NewObject<UUVEditorSeamToolProperties>();
+	Settings->RestoreProperties(this);
+	Settings->WatchProperty(Settings->Mode, [&](EUVEditorSeamMode) { OnSeamModeChanged(); });
+	AddToolPropertySource(Settings);
+
 	UContextObjectStore* ContextStore = GetToolManager()->GetContextObjectStore();
 	EmitChangeAPI = ContextStore->FindContext<UUVToolEmitChangeAPI>();
 	LivePreviewAPI = ContextStore->FindContext<UUVToolLivePreviewAPI>();
@@ -349,6 +402,16 @@ void UUVEditorSeamTool::Setup()
 	{
 		ClearHover();
 	};
+	UnwrapHoverBehavior->Modifiers.RegisterModifier(TemporaryModeToggleModifierID, FInputDeviceState::IsShiftKeyDown);
+	UnwrapHoverBehavior->OnUpdateModifierStateFunc = [this](int ModifierID, bool bIsOn) {
+		if (ModifierID == TemporaryModeToggleModifierID && bModeIsTemporarilyToggled != bIsOn)
+		{
+			bModeIsTemporarilyToggled = bIsOn;
+			ClearHover();
+			UpdateHover();
+			ResetPreviewColors();		
+		}
+	};
 	AddInputBehavior(UnwrapHoverBehavior);
 
 	// Now the applied/live preview viewport...
@@ -368,6 +431,7 @@ void UUVEditorSeamTool::Setup()
 			OnMeshVertexClicked(Vid, IndexOfMesh, false);
 		}
 	};
+
 	LivePreviewBehaviorSet->Add(LivePreviewClickBehavior, this);
 
 	ULocalMouseHoverBehavior* LivePreviewHoverBehavior = NewObject<ULocalMouseHoverBehavior>();
@@ -393,6 +457,17 @@ void UUVEditorSeamTool::Setup()
 	{
 		ClearHover();
 	};
+	LivePreviewHoverBehavior->Modifiers.RegisterModifier(TemporaryModeToggleModifierID, FInputDeviceState::IsShiftKeyDown);
+	LivePreviewHoverBehavior->OnUpdateModifierStateFunc = [this](int ModifierID, bool bIsOn)
+	{
+		if (ModifierID == TemporaryModeToggleModifierID && bModeIsTemporarilyToggled != bIsOn)
+		{
+			bModeIsTemporarilyToggled = bIsOn;
+			ClearHover();
+			UpdateHover();
+			ResetPreviewColors();
+		}
+	};
 	LivePreviewBehaviorSet->Add(LivePreviewHoverBehavior, this);
 
 	// Give the live preview behaviors to the live preview input router
@@ -526,19 +601,19 @@ void UUVEditorSeamTool::ReconstructLockedPathVisualization()
 		UUVEditorToolMeshInput* ClickedTarget = Targets[ClickedMeshIndex];
 
 		AddDisplayedPath(ClickedTarget, UnwrapLineSet, AppliedLineSet, LockedPath,
-			false, FUVEditorUXSettings::ToolLockedPathColor,
+			false, IsInJoinMode(), GetLockedPathColor(),
 			FUVEditorUXSettings::ToolLockedPathThickness,
 			FUVEditorUXSettings::ToolLockedPathDepthBias);
 
 		TArray<int32> UnwrapVids;
 		ClickedTarget->AppliedVidToUnwrapVids(LastLockedAppliedVid, UnwrapVids);
 		AddDisplayedPoints(ClickedTarget, UnwrapPointSet, AppliedPointSet, LastLockedAppliedVid, UnwrapVids, 
-			FUVEditorUXSettings::ToolLockedPathColor, FUVEditorUXSettings::ToolLockedPathDepthBias);
+			GetLockedPathColor(), FUVEditorUXSettings::ToolLockedPathDepthBias);
 
 		UnwrapVids.Reset();
 		ClickedTarget->AppliedVidToUnwrapVids(SeamStartAppliedVid, UnwrapVids);
 		AddDisplayedPoints(ClickedTarget, UnwrapStartPoints, AppliedStartPoints, SeamStartAppliedVid, UnwrapVids,
-			FUVEditorUXSettings::ToolLockedPathColor, FUVEditorUXSettings::ToolLockedPathDepthBias);
+			GetLockedPathColor(), FUVEditorUXSettings::ToolLockedPathDepthBias);
 	}
 }
 
@@ -575,7 +650,7 @@ void UUVEditorSeamTool::Shutdown(EToolShutdownType ShutdownType)
 
 	// Apply any pending seam if needed
 	if (ShutdownType != EToolShutdownType::Cancel && LockedPath.Num() > 0)
-	{
+	{		
 		ApplySeam(LockedPath);
 	}
 
@@ -659,7 +734,7 @@ void UUVEditorSeamTool::UpdateHover()
 	{
 		// Just draw the point and finish
 		AddDisplayedPoints(Targets[HoverMeshIndex], UnwrapPointSet, AppliedPointSet,
-			AppliedVid, UnwrapVids, FUVEditorUXSettings::ToolExtendPathColor, FUVEditorUXSettings::ToolExtendPathDepthBias);
+			AppliedVid, UnwrapVids, GetExtendPathColor(), FUVEditorUXSettings::ToolExtendPathDepthBias);
 		return;
 	}
 
@@ -671,12 +746,14 @@ void UUVEditorSeamTool::UpdateHover()
 		{
 			TArray<int32> LastLockedUnwrapVids;
 			Targets[HoverMeshIndex]->AppliedVidToUnwrapVids(LastLockedAppliedVid, LastLockedUnwrapVids);
-			GetVidPath(Targets[HoverMeshIndex]->UnwrapCanonical.Get(), LastLockedUnwrapVids, HoverVid, NewPathVids);
+			GetVidPath(Targets[HoverMeshIndex]->UnwrapCanonical.Get(), LastLockedUnwrapVids, HoverVid, NewPathVids,
+				       GetUVDistanceFunction(*Targets[HoverMeshIndex], bHoverVidIsFromUnwrap, IsInJoinMode()));
 		}
 		else
 		{
 			ensure(LastLockedAppliedVid != IndexConstants::InvalidID);
-			GetVidPath(Targets[HoverMeshIndex]->AppliedCanonical.Get(), LastLockedAppliedVid, HoverVid, NewPathVids);
+			GetVidPath(Targets[HoverMeshIndex]->AppliedCanonical.Get(), LastLockedAppliedVid, HoverVid, NewPathVids,
+				       GetUVDistanceFunction(*Targets[HoverMeshIndex], bHoverVidIsFromUnwrap, IsInJoinMode()));
 		}
 	}
 
@@ -686,7 +763,7 @@ void UUVEditorSeamTool::UpdateHover()
 		ULineSetComponent* UnwrapLineSet = UnwrapGeometry->FindLineSet(HoverLineSetID);
 		ULineSetComponent* AppliedLineSet = LivePreviewGeometry->FindLineSet(HoverLineSetID);
 		AddDisplayedPath(Targets[HoverMeshIndex], UnwrapLineSet, AppliedLineSet,
-			NewPathVids, bHoverVidIsFromUnwrap, FUVEditorUXSettings::ToolExtendPathColor,
+			NewPathVids, bHoverVidIsFromUnwrap, IsInJoinMode(), GetExtendPathColor(),
 			FUVEditorUXSettings::ToolExtendPathThickness, FUVEditorUXSettings::ToolExtendPathDepthBias);
 	}
 
@@ -702,7 +779,7 @@ void UUVEditorSeamTool::UpdateHover()
 		// there because that turns out to be confusing, esp when we have overlapping islands and
 		// the point is on another island.
 		AddDisplayedPoints(Targets[HoverMeshIndex], UnwrapPointSet, AppliedPointSet,
-			AppliedVid, UnwrapVids, FUVEditorUXSettings::ToolExtendPathColor, FUVEditorUXSettings::ToolExtendPathDepthBias);
+			AppliedVid, UnwrapVids, GetExtendPathColor(), FUVEditorUXSettings::ToolExtendPathDepthBias);
 	}
 }
 
@@ -735,7 +812,7 @@ void UUVEditorSeamTool::ResetPreviewColors()
 	using namespace UVEditorSeamToolLocals;
 
 	FColor Color = bCompletionColorOverride ? FUVEditorUXSettings::ToolCompletionPathColor 
-		: FUVEditorUXSettings::ToolExtendPathColor;
+		: GetExtendPathColor();
 
 	UnwrapGeometry->FindPointSet(HoverPointSetID)->SetAllPointsColor(Color);
 	LivePreviewGeometry->FindPointSet(HoverPointSetID)->SetAllPointsColor(Color);
@@ -743,7 +820,7 @@ void UUVEditorSeamTool::ResetPreviewColors()
 	LivePreviewGeometry->FindLineSet(HoverLineSetID)->SetAllLinesColor(Color);
 
 	Color = bCompletionColorOverride ? FUVEditorUXSettings::ToolCompletionPathColor 
-		: FUVEditorUXSettings::ToolLockedPathColor;
+		: GetLockedPathColor();
 
 	UnwrapGeometry->FindPointSet(LockedPointSetID)->SetAllPointsColor(Color);
 	LivePreviewGeometry->FindPointSet(LockedPointSetID)->SetAllPointsColor(Color);
@@ -776,7 +853,7 @@ void UUVEditorSeamTool::ApplyClick()
 		UPointSetComponent* UnwrapPointSet = UnwrapGeometry->FindPointSet(StartPointSetID);
 		UPointSetComponent* AppliedPointSet = LivePreviewGeometry->FindPointSet(StartPointSetID);
 		AddDisplayedPoints(ClickedTarget, UnwrapPointSet, AppliedPointSet, AppliedVid, UnwrapVids, 
-			FUVEditorUXSettings::ToolLockedPathColor, FUVEditorUXSettings::ToolLockedPathDepthBias);
+			GetLockedPathColor(), FUVEditorUXSettings::ToolLockedPathDepthBias);
 
 		LockedPath.Add(AppliedVid);
 		EmitChangeAPI->EmitToolDependentChange(this, 
@@ -797,7 +874,8 @@ void UUVEditorSeamTool::ApplyClick()
 	{
 		TArray<int32> LastLockedUnwrapVids;
 		ClickedTarget->AppliedVidToUnwrapVids(LastLockedAppliedVid, LastLockedUnwrapVids);
-		GetVidPath(ClickedTarget->UnwrapCanonical.Get(), LastLockedUnwrapVids, ClickedVid, VidPath);
+		GetVidPath(ClickedTarget->UnwrapCanonical.Get(), LastLockedUnwrapVids, ClickedVid, VidPath,
+			       GetUVDistanceFunction(*ClickedTarget, bClickWasInUnwrap, IsInJoinMode()));
 		for (int32 Vid : VidPath)
 		{
 			TempStorage.Add(ClickedTarget->UnwrapVidToAppliedVid(Vid));
@@ -806,7 +884,8 @@ void UUVEditorSeamTool::ApplyClick()
 	}
 	else
 	{
-		GetVidPath(ClickedTarget->AppliedCanonical.Get(), LastLockedAppliedVid, ClickedVid, VidPath);
+		GetVidPath(ClickedTarget->AppliedCanonical.Get(), LastLockedAppliedVid, ClickedVid, VidPath,
+			       GetUVDistanceFunction(*ClickedTarget, bClickWasInUnwrap, IsInJoinMode()));
 	}
 
 	if (VidPath.Num() == 0)
@@ -844,13 +923,13 @@ void UUVEditorSeamTool::ApplyClick()
 	ULineSetComponent* UnwrapLineSet = UnwrapGeometry->FindLineSet(LockedLineSetID);
 	ULineSetComponent* AppliedLineSet = LivePreviewGeometry->FindLineSet(LockedLineSetID);
 	AddDisplayedPath(ClickedTarget, UnwrapLineSet, AppliedLineSet, VidPath,
-		bClickWasInUnwrap, FUVEditorUXSettings::ToolLockedPathColor,
+		bClickWasInUnwrap, IsInJoinMode(), GetLockedPathColor(),
 		FUVEditorUXSettings::ToolLockedPathThickness,
 		FUVEditorUXSettings::ToolLockedPathDepthBias);
 	UnwrapPointSet->Clear();
 	AppliedPointSet->Clear();
 	AddDisplayedPoints(ClickedTarget, UnwrapPointSet, AppliedPointSet, AppliedVid, UnwrapVids, 
-		FUVEditorUXSettings::ToolLockedPathColor, FUVEditorUXSettings::ToolLockedPathDepthBias);
+		GetLockedPathColor(), FUVEditorUXSettings::ToolLockedPathDepthBias);
 	
 	EmitChangeAPI->EmitToolDependentChange(this, 
 		MakeUnique<FPathChange>(*AppliedVidPath, false, ClickedMeshIndex, true),
@@ -866,9 +945,13 @@ void UUVEditorSeamTool::ApplySeam(const TArray<int32>& AppliedVidsIn)
 	using namespace UVEditorSeamToolLocals;
 
 	UUVEditorToolMeshInput* ClickedTarget = Targets[ClickedMeshIndex];
+	
+	FDynamicMeshUVOverlay* UVOverlay = ClickedTarget->AppliedCanonical.Get()->Attributes()->GetUVLayer(ClickedTarget->UVLayerIndex);
+	FDynamicMeshUVEditor UVEditor(ClickedTarget->AppliedCanonical.Get(), UVOverlay);
 
-	// The locked path becomes an actual UV seam.
-
+	// The locked path (represented by AppliedVidsIn) is used to create an actual UV seam or join existing
+	// seams, depending on the current tool mode.
+	
 	// Convert vid path into set of eids path in the applied mesh
 	TSet<int32> EidSet;
 	for (int32 i = 0; i < AppliedVidsIn.Num() - 1; ++i)
@@ -876,21 +959,36 @@ void UUVEditorSeamTool::ApplySeam(const TArray<int32>& AppliedVidsIn)
 		int32 Eid = ClickedTarget->AppliedCanonical->FindEdge(AppliedVidsIn[i], AppliedVidsIn[i + 1]);
 		if (ensure(Eid != IndexConstants::InvalidID))
 		{
-			EidSet.Add(ClickedTarget->AppliedCanonical->FindEdge(AppliedVidsIn[i], AppliedVidsIn[i + 1]));
+			EidSet.Add(Eid);
 		}
 	}
 
-	FUVEditResult UVEditResult;
-	FDynamicMeshUVEditor UVEditor(ClickedTarget->AppliedCanonical.Get(), 
-		ClickedTarget->UVLayerIndex, false);
-	UVEditor.CreateSeamsAtEdges(EidSet, &UVEditResult);
+	TSet<int32> TidSet;	
+	TArray<int32> VidSet;
+	if (IsInJoinMode())
+	{
+		UVEditor.RemoveSeamsAtEdges(EidSet);
 
-	TSet<int32> TidSet;
-	for (int32 UnwrapVid : UVEditResult.NewUVElements)
+		// Don't actually need to update VidSet because no element values changed or were created, only connectivity.
+		for (int32 Vid : AppliedVidsIn)
+		{
+			TArray<int32> VertTids;
+			ClickedTarget->AppliedCanonical->GetVtxTriangles(Vid, VertTids);
+			TidSet.Append(VertTids);
+		}
+	}
+	else
 	{
-		TArray<int32> VertTids;
-		ClickedTarget->AppliedCanonical->GetVtxTriangles(ClickedTarget->UnwrapVidToAppliedVid(UnwrapVid), VertTids);
-		TidSet.Append(VertTids);
+		FUVEditResult UVEditResult;
+		UVEditor.CreateSeamsAtEdges(EidSet, &UVEditResult);
+
+		for (int32 UnwrapVid : UVEditResult.NewUVElements)
+		{
+			TArray<int32> VertTids;
+			ClickedTarget->AppliedCanonical->GetVtxTriangles(ClickedTarget->UnwrapVidToAppliedVid(UnwrapVid), VertTids);
+			TidSet.Append(VertTids);
+		}
+		VidSet = UVEditResult.NewUVElements;
 	}
 
 	FDynamicMeshChangeTracker ChangeTracker(ClickedTarget->UnwrapCanonical.Get());
@@ -898,7 +996,7 @@ void UUVEditorSeamTool::ApplySeam(const TArray<int32>& AppliedVidsIn)
 	ChangeTracker.SaveTriangles(TidSet, true);
 
 	TArray<int32> AppliedTids = TidSet.Array();
-	ClickedTarget->UpdateAllFromAppliedCanonical(&UVEditResult.NewUVElements, &AppliedTids, &AppliedTids);
+	ClickedTarget->UpdateAllFromAppliedCanonical(&VidSet, &AppliedTids, &AppliedTids);
 	Spatials2D[ClickedMeshIndex]->Build();
 
 	// Emit transaction
@@ -984,13 +1082,36 @@ int32 UUVEditorSeamTool::Get2DHitVertex(const FRay& WorldRayIn, int32* IndexOf2D
 	}
 	else if (State == EState::SeamInProgress)
 	{
-		// We don't provide a tolerance here because we want to snap to nearest vertex even if 
+		// This part goes about the process slightly differently than above, because in cases
+		// of drawing internal seams, we want to pick not just the nearest vertex to the mouse
+		// position, but the nearest vertex on the nearest triangle. Otherwise, we might end up
+		// picking the vertex on the other side of a recently drawn seam (which shares the same
+		// spatial location) and forcing the path finding to route the long way around the mesh
+		// to get to it.
+
+		// We could call FindNearestTriangle followed by FindNearestVertex with the triangle set
+		// as a filter, but for three vertices, it's probably faster to simply brute force the check.
+		
+		// We don't provide a tolerance here because we want to snap to nearest triangle and vertex even if 
 		// we're far outside the mesh
 		double Distance = -1;
-		int32 Vid = Spatials2D[ClickedMeshIndex]->FindNearestVertex(PlanePoint, Distance);
+		int32 Tid = Spatials2D[ClickedMeshIndex]->FindNearestTriangle(PlanePoint, Distance);
 		if (Distance >= 0)
 		{
-			Result = Vid;
+			TArray<FVector3d> VertPositions;
+			VertPositions.SetNum(3);
+			Targets[ClickedMeshIndex]->UnwrapCanonical->GetTriVertices(Tid, VertPositions[0], VertPositions[1], VertPositions[2]);
+			FIndex3i Vids = Targets[ClickedMeshIndex]->UnwrapCanonical->GetTriangle(Tid);
+			Distance = TNumericLimits<double>::Max();
+			for (uint32 VIndex = 0; VIndex < 3; ++VIndex)
+			{
+				if (FVector3d::DistSquared(VertPositions[VIndex], PlanePoint) < Distance)
+				{
+					Distance = FVector3d::DistSquared(VertPositions[VIndex], PlanePoint);
+					Result = Vids[VIndex];
+				}
+			}
+
 			if (IndexOf2DSpatialOut)
 			{
 				*IndexOf2DSpatialOut = ClickedMeshIndex;
@@ -1187,6 +1308,31 @@ bool UUVEditorSeamTool::ExecuteNestedAcceptCommand()
 	return false;
 }
 
+FColor UUVEditorSeamTool::GetLockedPathColor() const
+{
+	return !IsInJoinMode() ? FUVEditorUXSettings::ToolLockedCutPathColor :
+     	                     FUVEditorUXSettings::ToolLockedJoinPathColor;
+}
+
+
+FColor UUVEditorSeamTool::GetExtendPathColor() const
+{
+	return !IsInJoinMode() ? FUVEditorUXSettings::ToolExtendCutPathColor :
+    	                     FUVEditorUXSettings::ToolExtendJoinPathColor;
+}
+
+bool UUVEditorSeamTool::IsInJoinMode() const
+{
+	bool bIsJoinMode = Settings && Settings->Mode == EUVEditorSeamMode::Join;
+	return bModeIsTemporarilyToggled ? !bIsJoinMode : bIsJoinMode;
+}
+
+void UUVEditorSeamTool::OnSeamModeChanged()
+{
+	bModeIsTemporarilyToggled = false;
+	ResetPreviewColors();
+}
+
 void UUVEditorSeamTool::RecordAnalytics()
 {
 	using namespace UVEditorAnalytics;
diff --git a/Engine/Plugins/Experimental/UVEditor/Source/UVEditorTools/Private/UVEditorStyle.cpp b/Engine/Plugins/Experimental/UVEditor/Source/UVEditorTools/Private/UVEditorStyle.cpp
--- a/Engine/Plugins/Experimental/UVEditor/Source/UVEditorTools/Private/UVEditorStyle.cpp
+++ b/Engine/Plugins/Experimental/UVEditor/Source/UVEditorTools/Private/UVEditorStyle.cpp
@@ -6,6 +6,7 @@
 #include "Styling/CoreStyle.h"
 #include "Styling/AppStyle.h"
 #include "Styling/SlateStyleMacros.h"
+#include "Styling/ToolBarStyle.h"
 
 FName FUVEditorStyle::StyleName("UVStyle");
 
@@ -72,7 +73,24 @@ FUVEditorStyle::FUVEditorStyle()
 	Set("UVEditor.AlignDirectionCentersHorizontally", new IMAGE_BRUSH_SVG("AlignCentersHorizontal", ToolbarIconSize));
 	Set("UVEditor.AlignDirectionCentersVertically", new IMAGE_BRUSH_SVG("AlignCentersVertical", ToolbarIconSize));
 
+	// Style for the toolbar in the SeamTool customization
+	{
+		// For the selection button toolbar, we want to use something similar to the toolbar we use in the viewport
+		Set("SeamTool.ModeToolbar", FAppStyle::Get().GetWidgetStyle<FToolBarStyle>("EditorViewportToolBar"));
 
+		// However, increase the size of the buttons a bit
+		FCheckBoxStyle ToggleButtonStart = FAppStyle::Get().GetWidgetStyle<FCheckBoxStyle>("EditorViewportToolBar.ToggleButton.Start");
+		ToggleButtonStart.SetPadding(FMargin(9, 7, 6, 7));
+		Set("SeamTool.ModeToolbar.ToggleButton.Start", ToggleButtonStart);
+
+		FCheckBoxStyle ToggleButtonMiddle = FAppStyle::Get().GetWidgetStyle<FCheckBoxStyle>("EditorViewportToolBar.ToggleButton.Middle");
+		ToggleButtonMiddle.SetPadding(FMargin(9, 7, 6, 7));
+		Set("SeamTool.ModeToolbar.ToggleButton.Middle", ToggleButtonMiddle);
+
+		FCheckBoxStyle ToggleButtonEnd = FAppStyle::Get().GetWidgetStyle<FCheckBoxStyle>("EditorViewportToolBar.ToggleButton.End");
+		ToggleButtonEnd.SetPadding(FMargin(7, 7, 8, 7));
+		Set("SeamTool.ModeToolbar.ToggleButton.End", ToggleButtonEnd);
+	}
 
 	FSlateStyleRegistry::RegisterSlateStyle(*this);
 }
diff --git a/Engine/Plugins/Experimental/UVEditor/Source/UVEditorTools/Private/UVEditorUXSettings.cpp b/Engine/Plugins/Experimental/UVEditor/Source/UVEditorTools/Private/UVEditorUXSettings.cpp
--- a/Engine/Plugins/Experimental/UVEditor/Source/UVEditorTools/Private/UVEditorUXSettings.cpp
+++ b/Engine/Plugins/Experimental/UVEditor/Source/UVEditorTools/Private/UVEditorUXSettings.cpp
@@ -13,9 +13,9 @@ const float FUVEditorUXSettings::CameraNearPlaneProportionZ(0.8); // Top layer,
 // Note: While these are floating point values, they represent percentages and should be separated
 // by at least integer amounts, as they serve double duty in certain cases for translucent primitive
 // sorting order.
-const float FUVEditorUXSettings::ToolLockedPathDepthBias(6.0);
-const float FUVEditorUXSettings::ToolExtendPathDepthBias(6.0);
-const float FUVEditorUXSettings::SewLineDepthOffset(5.0f);
+const float FUVEditorUXSettings::ToolLockedPathDepthBias(8.0);
+const float FUVEditorUXSettings::ToolExtendPathDepthBias(8.0);
+const float FUVEditorUXSettings::SewLineDepthOffset(7.0f);
 const float FUVEditorUXSettings::SelectionHoverWireframeDepthBias(6);
 const float FUVEditorUXSettings::SelectionHoverTriangleDepthBias(5);
 const float FUVEditorUXSettings::SelectionWireframeDepthBias(4.0);
@@ -55,8 +55,10 @@ const FColor FUVEditorUXSettings::SelectionHoverTriangleWireframeColor(FColor::F
 const FColor FUVEditorUXSettings::SewSideLeftColor(FColor::Red);
 const FColor FUVEditorUXSettings::SewSideRightColor(FColor::Green);
 
-const FColor FUVEditorUXSettings::ToolLockedPathColor(FColor::Green);
-const FColor FUVEditorUXSettings::ToolExtendPathColor(FColor::Green);
+const FColor FUVEditorUXSettings::ToolLockedCutPathColor(FColor::Green);
+const FColor FUVEditorUXSettings::ToolExtendCutPathColor(FColor::Green);
+const FColor FUVEditorUXSettings::ToolLockedJoinPathColor(FColor::Turquoise);
+const FColor FUVEditorUXSettings::ToolExtendJoinPathColor(FColor::Turquoise);
 const FColor FUVEditorUXSettings::ToolCompletionPathColor(FColor::Orange);
 
 const FColor FUVEditorUXSettings::LivePreviewExistingSeamColor(FColor::Green);
diff --git a/Engine/Plugins/Experimental/UVEditor/Source/UVEditorTools/Public/UVEditorSeamTool.h b/Engine/Plugins/Experimental/UVEditor/Source/UVEditorTools/Public/UVEditorSeamTool.h
--- a/Engine/Plugins/Experimental/UVEditor/Source/UVEditorTools/Public/UVEditorSeamTool.h
+++ b/Engine/Plugins/Experimental/UVEditor/Source/UVEditorTools/Public/UVEditorSeamTool.h
@@ -24,6 +24,26 @@ class UUVToolLivePreviewAPI;
 
 PREDECLARE_GEOMETRY(class FDynamicMesh);
 
+UENUM()
+enum class EUVEditorSeamMode : uint8
+{
+	/** Marked path will cut the UV island, creating new seams.*/	
+	Cut = 0,
+	/** Marked path will join the UV island, removing seams.*/
+	Join = 1
+};
+
+UCLASS()
+class UVEDITORTOOLS_API UUVEditorSeamToolProperties : public UInteractiveToolPropertySet
+{
+	GENERATED_BODY()
+
+public:
+
+	UPROPERTY(EditAnywhere, Category = Options)
+	EUVEditorSeamMode Mode = EUVEditorSeamMode::Cut;
+};
+
 UCLASS()
 class UVEDITORTOOLS_API UUVEditorSeamToolBuilder : public UInteractiveToolBuilder
 {
@@ -87,6 +107,11 @@ protected:
 
 	void UpdateToolMessage();
 
+	int32 TemporaryModeToggleModifierID = 1;
+	bool bModeIsTemporarilyToggled;
+	bool IsInJoinMode() const;
+	void OnSeamModeChanged();
+
 	enum class EState
 	{
 		WaitingToStart,
@@ -98,6 +123,9 @@ protected:
 	UPROPERTY()
 	TArray<TObjectPtr<UUVEditorToolMeshInput>> Targets;
 
+	UPROPERTY()
+	TObjectPtr< UUVEditorSeamToolProperties> Settings = nullptr;
+
 	UPROPERTY()
 	TObjectPtr<UUVToolLivePreviewAPI> LivePreviewAPI = nullptr;
 
@@ -143,6 +171,9 @@ protected:
 	TArray<int32> LockedPath;
 	TArray<int32> PreviewPath;
 
+	FColor GetLockedPathColor() const;
+	FColor GetExtendPathColor() const;
+
 	// When true, the entire path is changed to the "completion" color to show
 	// that the next click will complete the path.
 	bool bCompletionColorOverride = false;
diff --git a/Engine/Plugins/Experimental/UVEditor/Source/UVEditorTools/Public/UVEditorUXSettings.h b/Engine/Plugins/Experimental/UVEditor/Source/UVEditorTools/Public/UVEditorUXSettings.h
--- a/Engine/Plugins/Experimental/UVEditor/Source/UVEditorTools/Public/UVEditorUXSettings.h
+++ b/Engine/Plugins/Experimental/UVEditor/Source/UVEditorTools/Public/UVEditorUXSettings.h
@@ -68,10 +68,12 @@ public:
 
 	// These are currently used by the seam tool but can be generally used by
 	// tools for displaying paths.
-	static const FColor ToolLockedPathColor;
+	static const FColor ToolLockedCutPathColor;
+	static const FColor ToolLockedJoinPathColor;
 	static const float ToolLockedPathThickness;
 	static const float ToolLockedPathDepthBias;
-	static const FColor ToolExtendPathColor;
+	static const FColor ToolExtendCutPathColor;
+	static const FColor ToolExtendJoinPathColor;
 	static const float ToolExtendPathThickness;
 	static const float ToolExtendPathDepthBias;
 	static const FColor ToolCompletionPathColor;
diff --git a/Engine/Plugins/Experimental/UVEditor/Source/UVEditorToolsEditorOnly/Private/DetailsCustomizations/UVEditorSeamToolCustomizations.cpp b/Engine/Plugins/Experimental/UVEditor/Source/UVEditorToolsEditorOnly/Private/DetailsCustomizations/UVEditorSeamToolCustomizations.cpp
new file mode 100644
--- /dev/null
+++ b/Engine/Plugins/Experimental/UVEditor/Source/UVEditorToolsEditorOnly/Private/DetailsCustomizations/UVEditorSeamToolCustomizations.cpp
@@ -0,0 +1,199 @@
+// Copyright Epic Games, Inc. All Rights Reserved.
+
+#include "DetailsCustomizations/UVEditorSeamToolCustomizations.h"
+#include "DetailCategoryBuilder.h"
+#include "DetailLayoutBuilder.h"
+#include "DetailWidgetRow.h"
+#include "PropertyCustomizationHelpers.h"
+#include "Widgets/Input/SButton.h"
+#include "Widgets/Input/SComboButton.h"
+#include "Widgets/Layout/SUniformGridPanel.h"
+#include "Widgets/Layout/SUniformWrapPanel.h"
+#include "Widgets/Images/SLayeredImage.h"
+#include "IDetailChildrenBuilder.h"
+#include "Internationalization/BreakIterator.h"
+#include "Framework/MultiBox/MultiBoxBuilder.h"
+#include "IContentBrowserSingleton.h"
+#include "SAssetView.h"
+#include "Widgets/Input/SSegmentedControl.h"
+
+#include "UVEditorStyle.h"
+#include "UVEditorSeamTool.h"
+
+
+#define LOCTEXT_NAMESPACE "UVEditorSeamToolDetailsCustomization"
+
+namespace FUVEditorSeamToolPropertiesDetailsLocal
+{
+	template<class UENUM_TYPE>
+	FPropertyAccess::Result GetPropertyValueAsEnum(const TSharedPtr<IPropertyHandle> Property, UENUM_TYPE& Value)
+	{
+		FPropertyAccess::Result Result;
+		if (Property.IsValid())
+		{
+			uint8 ValueAsInt8;
+			Result = Property->GetValue(/*out*/ ValueAsInt8);
+			if (Result == FPropertyAccess::Success)
+			{
+				Value = (UENUM_TYPE)ValueAsInt8;
+				return FPropertyAccess::Success;
+			}
+			uint16 ValueAsInt16;
+			Result = Property->GetValue(/*out*/ ValueAsInt16);
+			if (Result == FPropertyAccess::Success)
+			{
+				Value = (UENUM_TYPE)ValueAsInt16;
+				return FPropertyAccess::Success;
+			}
+			uint32 ValueAsInt32;
+			Result = Property->GetValue(/*out*/ ValueAsInt32);
+			if (Result == FPropertyAccess::Success)
+			{
+				Value = (UENUM_TYPE)ValueAsInt32;
+				return FPropertyAccess::Success;
+			}
+		}
+		return FPropertyAccess::Fail;
+	}
+
+	template<class UENUM_TYPE>
+	FPropertyAccess::Result SetPropertyValueAsEnum(const TSharedPtr<IPropertyHandle> Property, const UENUM_TYPE& Value)
+	{
+		FPropertyAccess::Result Result;
+
+		if (Property.IsValid())
+		{
+			uint32 ValueAsInt32 = (uint32)Value;;
+			Result = Property->SetValue(ValueAsInt32);
+			if (Result == FPropertyAccess::Success)
+			{
+				return Result;
+			}
+			uint16 ValueAsInt16 = (uint16)Value;;
+			Result = Property->SetValue(ValueAsInt16);
+			if (Result == FPropertyAccess::Success)
+			{
+				return Result;
+			}
+			uint8 ValueAsInt8 = (uint8)Value;;
+			Result = Property->SetValue(ValueAsInt8);
+			if (Result == FPropertyAccess::Success)
+			{
+				return Result;
+			}
+		}
+		return FPropertyAccess::Fail;
+	}
+
+	template<class PROPERTIES_TYPE, class TOOLBARBUILDER_TYPE, class UENUM_TYPE>
+	void AddToggleButtonForEnum(PROPERTIES_TYPE& Properties, TOOLBARBUILDER_TYPE& ToolbarBuilder, TSharedPtr<IPropertyHandle> EnumPropertyHandle, UENUM_TYPE ButtonValue, FText Label, FName IconName)
+	{
+		ToolbarBuilder.AddToolBarButton(FUIAction(
+			FExecuteAction::CreateWeakLambda(Properties, [EnumPropertyHandle, ButtonValue]()
+			{
+				SetPropertyValueAsEnum(EnumPropertyHandle, ButtonValue);
+			}),
+			FCanExecuteAction::CreateWeakLambda(Properties, [EnumPropertyHandle]() {
+				return EnumPropertyHandle->IsEditable();
+			}),
+			FIsActionChecked::CreateWeakLambda(Properties, [EnumPropertyHandle, ButtonValue]()
+			{
+				UENUM_TYPE Value;
+				GetPropertyValueAsEnum(EnumPropertyHandle, Value);
+				return Value == ButtonValue;
+			})),
+			NAME_None, // Extension hook
+			TAttribute<FText>(Label), // Label
+			TAttribute<FText>(Label), // Tooltip
+			TAttribute<FSlateIcon>(FSlateIcon(FUVEditorStyle::Get().GetStyleSetName(), IconName)),
+			EUserInterfaceActionType::ToggleButton);
+	};
+
+}
+
+
+//
+// UVEditorTransformTool
+//
+
+
+TSharedRef<IDetailCustomization> FUVEditorSeamToolPropertiesDetails::MakeInstance()
+{
+	return MakeShareable(new FUVEditorSeamToolPropertiesDetails);
+}
+
+
+void FUVEditorSeamToolPropertiesDetails::CustomizeDetails(IDetailLayoutBuilder& DetailBuilder)
+{
+	using namespace FUVEditorSeamToolPropertiesDetailsLocal;
+
+	TArray<TWeakObjectPtr<UObject>> ObjectsBeingCustomized;
+	DetailBuilder.GetObjectsBeingCustomized(ObjectsBeingCustomized);
+	check(ObjectsBeingCustomized.Num() > 0);
+	UUVEditorSeamToolProperties* Properties = CastChecked<UUVEditorSeamToolProperties>(ObjectsBeingCustomized[0]);
+
+	const FName OptionsCategoryName = "Options";
+	IDetailCategoryBuilder& CategoryBuilder = DetailBuilder.EditCategory(OptionsCategoryName);
+
+	TSharedRef<IPropertyHandle> ModeHandle = DetailBuilder.GetProperty(GET_MEMBER_NAME_CHECKED(UUVEditorSeamToolProperties, Mode), UUVEditorSeamToolProperties::StaticClass());
+	ensure(ModeHandle->IsValidHandle());
+	ModeHandle->MarkHiddenByCustomization();
+
+	// Create a toolbar for the selection filter
+	TSharedRef<FUICommandList> CommandList = MakeShared<FUICommandList>();
+	FToolBarBuilder ToolbarBuilder(CommandList, FMultiBoxCustomization::None);
+
+	FName ToolBarStyle = "SeamTool.ModeToolbar";
+	ToolbarBuilder.SetStyle(&FUVEditorStyle::Get(), ToolBarStyle);
+	ToolbarBuilder.SetLabelVisibility(EVisibility::Visible);
+
+	ToolbarBuilder.BeginSection("Tool Mode");
+	ToolbarBuilder.BeginBlockGroup();
+
+	FUVEditorSeamToolPropertiesDetailsLocal::AddToggleButtonForEnum(Properties, ToolbarBuilder, ModeHandle, EUVEditorSeamMode::Cut, LOCTEXT("SeamToolSplitMode", "Split"), "UVEditor.SplitAction");
+	FUVEditorSeamToolPropertiesDetailsLocal::AddToggleButtonForEnum(Properties, ToolbarBuilder, ModeHandle, EUVEditorSeamMode::Join, LOCTEXT("SeamToolSewMode", "Sew"), "UVEditor.SewAction");
+
+	ToolbarBuilder.EndBlockGroup();
+	ToolbarBuilder.EndSection();
+
+	FDetailWidgetRow& CustomRow = CategoryBuilder.AddCustomRow(FText::GetEmpty())
+	.NameContent()
+	[
+		SNew(STextBlock)
+		.Text(LOCTEXT("SeamToolModeLabel", "Mode"))
+	]
+	.ValueContent()
+	[
+	
+		SNew(SBox)
+		[
+			SNew(SSegmentedControl<EUVEditorSeamMode>)
+			.Value(this, &FUVEditorSeamToolPropertiesDetails::GetCurrentMode, ModeHandle)
+			.OnValueChanged(this, &FUVEditorSeamToolPropertiesDetails::OnCurrentModeChanged, ModeHandle)
+			+ SSegmentedControl<EUVEditorSeamMode>::Slot(EUVEditorSeamMode::Cut)
+		        .Text(LOCTEXT("SeamToolSplitMode", "Split"))
+				.Icon(FUVEditorStyle::Get().GetBrush("UVEditor.SplitAction"))
+				.ToolTip(LOCTEXT("SeamToolSplitMode", "Split"))
+			+ SSegmentedControl<EUVEditorSeamMode>::Slot(EUVEditorSeamMode::Join)
+		        .Text(LOCTEXT("SeamToolSewMode", "Sew"))
+				.Icon(FUVEditorStyle::Get().GetBrush("UVEditor.SewAction"))
+				.ToolTip(LOCTEXT("SeamToolSewMode", "Sew"))
+		]
+	];
+
+}
+
+
+EUVEditorSeamMode FUVEditorSeamToolPropertiesDetails::GetCurrentMode(TSharedRef<IPropertyHandle> PropertyHandle) const
+{
+	EUVEditorSeamMode ModeValue = EUVEditorSeamMode::Cut;
+	FUVEditorSeamToolPropertiesDetailsLocal::GetPropertyValueAsEnum(PropertyHandle, ModeValue);
+	return ModeValue;
+}
+
+void FUVEditorSeamToolPropertiesDetails::OnCurrentModeChanged(EUVEditorSeamMode ModeValue, TSharedRef<IPropertyHandle> PropertyHandle)
+{
+	FUVEditorSeamToolPropertiesDetailsLocal::SetPropertyValueAsEnum(PropertyHandle, ModeValue);
+}
+
+#undef LOCTEXT_NAMESPACE
\ No newline at end of file
diff --git a/Engine/Plugins/Experimental/UVEditor/Source/UVEditorToolsEditorOnly/Private/UVEditorToolsEditorOnlyModule.cpp b/Engine/Plugins/Experimental/UVEditor/Source/UVEditorToolsEditorOnly/Private/UVEditorToolsEditorOnlyModule.cpp
--- a/Engine/Plugins/Experimental/UVEditor/Source/UVEditorToolsEditorOnly/Private/UVEditorToolsEditorOnlyModule.cpp
+++ b/Engine/Plugins/Experimental/UVEditor/Source/UVEditorToolsEditorOnly/Private/UVEditorToolsEditorOnlyModule.cpp
@@ -6,9 +6,11 @@
 #include "PropertyEditorModule.h"
 
 #include "DetailsCustomizations/UVTransformToolCustomizations.h"
+#include "DetailsCustomizations/UVEditorSeamToolCustomizations.h"
 
 #include "UVEditorTransformTool.h"
 #include "Operators/UVEditorUVTransformOp.h"
+#include "UVEditorSeamTool.h"
 
 #define LOCTEXT_NAMESPACE "FUVEditorToolsEditorOnlyModule"
 
@@ -69,7 +71,10 @@ void FUVEditorToolsEditorOnlyModule::OnPostEngineInit()
 		FOnGetDetailCustomizationInstance::CreateStatic(&FUVEditorUVDistributeToolDetails::MakeInstance));
 	ClassesToUnregisterOnShutdown.Add(UUVEditorUVDistributeProperties::StaticClass()->GetFName());
 
-
+	// Seam Tool
+	PropertyModule.RegisterCustomClassLayout(UUVEditorSeamToolProperties::StaticClass()->GetFName(),
+		FOnGetDetailCustomizationInstance::CreateStatic(&FUVEditorSeamToolPropertiesDetails::MakeInstance));
+	ClassesToUnregisterOnShutdown.Add(UUVEditorSeamToolProperties::StaticClass()->GetFName());
 }
 
 #undef LOCTEXT_NAMESPACE
diff --git a/Engine/Plugins/Experimental/UVEditor/Source/UVEditorToolsEditorOnly/Public/DetailsCustomizations/UVEditorSeamToolCustomizations.h b/Engine/Plugins/Experimental/UVEditor/Source/UVEditorToolsEditorOnly/Public/DetailsCustomizations/UVEditorSeamToolCustomizations.h
new file mode 100644
--- /dev/null
+++ b/Engine/Plugins/Experimental/UVEditor/Source/UVEditorToolsEditorOnly/Public/DetailsCustomizations/UVEditorSeamToolCustomizations.h
@@ -0,0 +1,24 @@
+// Copyright Epic Games, Inc. All Rights Reserved.
+
+#pragma once
+
+#include "IDetailCustomization.h"
+#include "Input/Reply.h"
+#include "Templates/SharedPointer.h"
+#include "UObject/WeakObjectPtr.h"
+#include "UVEditorSeamTool.h"
+
+class IPropertyHandle;
+class FPropertyRestriction;
+class IUVEditorTransformToolQuickAction;
+
+// Customization for UUVEditorSeamToolProperties
+class FUVEditorSeamToolPropertiesDetails : public IDetailCustomization
+{
+public:
+	static TSharedRef<IDetailCustomization> MakeInstance();
+	virtual void CustomizeDetails(IDetailLayoutBuilder& DetailBuilder) override;
+
+	EUVEditorSeamMode GetCurrentMode(TSharedRef<IPropertyHandle> PropertyHandle) const;
+	void OnCurrentModeChanged(EUVEditorSeamMode Mode, TSharedRef<IPropertyHandle> PropertyHandle);
+};
\ No newline at end of file
diff --git a/Engine/Plugins/Runtime/GeometryProcessing/Source/DynamicMesh/Private/Parameterization/DynamicMeshUVEditor.cpp b/Engine/Plugins/Runtime/GeometryProcessing/Source/DynamicMesh/Private/Parameterization/DynamicMeshUVEditor.cpp
--- a/Engine/Plugins/Runtime/GeometryProcessing/Source/DynamicMesh/Private/Parameterization/DynamicMeshUVEditor.cpp
+++ b/Engine/Plugins/Runtime/GeometryProcessing/Source/DynamicMesh/Private/Parameterization/DynamicMeshUVEditor.cpp
@@ -817,6 +817,62 @@ static int32 FindUVElementForVertex(const FDynamicMesh3* Mesh, const FDynamicMes
 }
 
 
+bool FDynamicMeshUVEditor::RemoveSeamsAtEdges(const TSet<int32>& EidsToRemoveAsSeams)
+{
+	for (int32 Eid : EidsToRemoveAsSeams)
+	{
+		FIndex2i EdgeVids = Mesh->GetEdgeV(Eid);
+
+		// The following logic closely resemebles the logic within FDynamicMeshOverlay::IsSeamEdge,
+		// but since we want to know both if it's a seam and use many of the intermediate values
+		// computed by IsSeamEdge, we replicate it here.
+
+		FIndex2i Tris = Mesh->GetEdgeT(Eid);
+		if (Tris.B == FDynamicMesh3::InvalidID)
+		{
+			continue; // Technically a seam, but we don't want this one because there's no opposite edge to merge.
+		}
+
+		bool bASet = UVOverlay->IsSetTriangle(Tris.A), bBSet = UVOverlay->IsSetTriangle(Tris.B);
+		if (!bASet || !bBSet)
+		{
+			continue; // Similar problem as the above case - could be a seam here, but if one triangle
+					  // in the Overlay isn't set, there's nothing to merge.
+		}
+
+		FIndex3i Triangle0 = UVOverlay->GetTriangle(Tris.A);
+		FIndex3i BaseTriangle0 = Mesh->GetTriangle(Tris.A);
+		int idx_base_a0 = BaseTriangle0.IndexOf(EdgeVids.A);
+		int idx_base_b0 = BaseTriangle0.IndexOf(EdgeVids.B);
+
+		FIndex3i Triangle1 = UVOverlay->GetTriangle(Tris.B);
+		FIndex3i BaseTriangle1 = Mesh->GetTriangle(Tris.B);
+		int idx_base_a1 = BaseTriangle1.IndexOf(EdgeVids.A);
+		int idx_base_b1 = BaseTriangle1.IndexOf(EdgeVids.B);
+
+		int el_a_tri0 = Triangle0[idx_base_a0];
+		int el_b_tri0 = Triangle0[idx_base_b0];
+		int el_a_tri1 = Triangle1[idx_base_a1];
+		int el_b_tri1 = Triangle1[idx_base_b1];
+
+		// This shouldn't ever be the case, but lets just check. If true,
+		// it would indicate that there's a joined seam but somehow our pairwise
+		// vertex matching didn't work above for some reason.
+		ensure(!(el_a_tri0 == el_b_tri1 && el_b_tri0 == el_a_tri1));
+
+		if (el_a_tri0 != el_a_tri1)
+		{
+			UVOverlay->MergeElement(el_a_tri0, el_a_tri1);
+		}
+		if (el_b_tri0 != el_b_tri1)
+		{
+			UVOverlay->MergeElement(el_b_tri0, el_b_tri1);
+
+		}
+	}
+	return true;
+}
+
 
 bool FDynamicMeshUVEditor::CreateSeamsAtEdges(const TSet<int32>& EidsToMakeIntoSeams, FUVEditResult* Result)
 {
diff --git a/Engine/Plugins/Runtime/GeometryProcessing/Source/DynamicMesh/Public/Parameterization/DynamicMeshUVEditor.h b/Engine/Plugins/Runtime/GeometryProcessing/Source/DynamicMesh/Public/Parameterization/DynamicMeshUVEditor.h
--- a/Engine/Plugins/Runtime/GeometryProcessing/Source/DynamicMesh/Public/Parameterization/DynamicMeshUVEditor.h
+++ b/Engine/Plugins/Runtime/GeometryProcessing/Source/DynamicMesh/Public/Parameterization/DynamicMeshUVEditor.h
@@ -180,8 +180,16 @@ public:
 	 */
 	bool SetTriangleUVsFromFreeBoundarySpectralConformal(const TArray<int32>& Triangles, bool bUseExistingUVTopology, bool bPreserveIrregularity, FUVEditResult* Result = nullptr);
 
+	/** 
+	 * Merge existing UV topology with a set of edges, removing seams at edges if they exist within the UV topology.
+	 * 
+	 * @param EidsToRemoveAsSeams list of edges to remove seams from
+	 * @return true on success
+	 */
+	bool RemoveSeamsAtEdges(const TSet<int32>& EidsToRemoveAsSeams);
+
 	/**
-	 * Cut existing UV topolgy with a set of edges. This allows for creating partial seams/darts, interior cuts, etc.
+	 * Cut existing UV topology with a set of edges. This allows for creating partial seams/darts, interior cuts, etc.
 	 * 
 	 * Avoids creating bowties. In cases where an edge is not next to any present or future seams/borders, some
 	 * adjacent edge will be picked to be made into a seam as well, since it's impossible to make the original
diff --git a/Engine/Source/Runtime/GeometryCore/Private/DynamicMesh/DynamicMeshOverlay.cpp b/Engine/Source/Runtime/GeometryCore/Private/DynamicMesh/DynamicMeshOverlay.cpp
--- a/Engine/Source/Runtime/GeometryCore/Private/DynamicMesh/DynamicMeshOverlay.cpp
+++ b/Engine/Source/Runtime/GeometryCore/Private/DynamicMesh/DynamicMeshOverlay.cpp
@@ -203,6 +203,52 @@ void TDynamicMeshOverlay<RealType, ElementSize>::SplitVerticesWithPredicate(TFun
 }
 
 
+template<typename RealType, int ElementSize>
+bool TDynamicMeshOverlay<RealType, ElementSize>::MergeElement(int SourceElementID, int TargetElementID)
+{	
+	if (SourceElementID == TargetElementID)
+	{
+		return false;
+	}
+
+	int SourceParentID = ParentVertices[SourceElementID];
+	int TargetParentID = ParentVertices[TargetElementID];
+
+	auto MergeElementForTriangle = [this, SourceElementID, TargetElementID](int32 TriID)
+	{
+		int ElementTriStart = TriID * 3;
+		for (int SubIdx = 0; SubIdx < 3; SubIdx++)
+		{
+			int CurElID = ElementTriangles[ElementTriStart + SubIdx];
+			if (CurElID == SourceElementID)
+			{
+				ElementsRefCounts.Decrement(SourceElementID);
+				ElementsRefCounts.Increment(TargetElementID);
+				ElementTriangles[ElementTriStart + SubIdx] = TargetElementID;
+			}
+		}
+	};
+
+	checkSlow(SourceParentID == TargetParentID);
+	if (SourceParentID != TargetParentID)
+	{
+		return false;
+	}
+
+	ParentMesh->EnumerateVertexTriangles(SourceParentID, MergeElementForTriangle);
+
+	checkSlow(ElementsRefCounts.IsValid(SourceElementID));
+	checkSlow(ElementsRefCounts.GetRefCount(SourceElementID) == 1);
+	if (ElementsRefCounts.GetRefCount(SourceElementID) == 1)
+	{
+		ElementsRefCounts.Decrement(SourceElementID);
+		ParentVertices[SourceElementID] = FDynamicMesh3::InvalidID;
+	}
+
+	return true;
+}
+
+
 template<typename RealType, int ElementSize>
 int TDynamicMeshOverlay<RealType, ElementSize>::SplitElement(int ElementID, const TArrayView<const int>& TrianglesToUpdate)
 {
diff --git a/Engine/Source/Runtime/GeometryCore/Public/DynamicMesh/DynamicMeshOverlay.h b/Engine/Source/Runtime/GeometryCore/Public/DynamicMesh/DynamicMeshOverlay.h
--- a/Engine/Source/Runtime/GeometryCore/Public/DynamicMesh/DynamicMeshOverlay.h
+++ b/Engine/Source/Runtime/GeometryCore/Public/DynamicMesh/DynamicMeshOverlay.h
@@ -347,6 +347,14 @@ public:
 	 */
 	void SplitVerticesWithPredicate(TFunctionRef<bool(int ElementIdx, int TriID)> ShouldSplitOutVertex, TFunctionRef<void(int ElementIdx, int TriID, RealType* FillVect)> GetNewElementValue);
 
+	/**
+	 * Collapse SourceElementID into TargetElementID, resulting in connecting any containing triangles and reducing the total elements in the overlay.
+	 *
+	 * @param SourceElementID the element to merge away
+	 * @param TargetElementID the element to merge into
+	 * @return If the operation completed successfully, returns true
+	 */
+	bool MergeElement(int SourceElementID, int TargetElementID);
 
 	/**
 	 * Create a new copy of ElementID, and update connected triangles in the TrianglesToUpdate array to reference the copy of ElementID where they used to reference ElementID
diff --git a/Engine/Source/Runtime/InteractiveToolsFramework/Public/BaseBehaviors/MouseHoverBehavior.h b/Engine/Source/Runtime/InteractiveToolsFramework/Public/BaseBehaviors/MouseHoverBehavior.h
--- a/Engine/Source/Runtime/InteractiveToolsFramework/Public/BaseBehaviors/MouseHoverBehavior.h
+++ b/Engine/Source/Runtime/InteractiveToolsFramework/Public/BaseBehaviors/MouseHoverBehavior.h
@@ -83,6 +83,9 @@ public:
 	/** lambda implementation of OnEndHover */
 	TUniqueFunction<void()> OnEndHoverFunc = []() {};
 
+	/** lambda implementation of OnUpdateModifierState */
+	TUniqueFunction< void(int, bool) > OnUpdateModifierStateFunc = [](int ModifierID, bool bIsOn) {};
+
 public:
 	// IHoverBehaviorTarget implementation
 	virtual FInputRayHit BeginHoverSequenceHitTest(const FInputDeviceRay& PressPos) override
@@ -104,4 +107,10 @@ public:
 	{
 		OnEndHoverFunc();
 	}
+
+	// IModifierToggleBehaviorTarget implementation
+	virtual void OnUpdateModifierState(int ModifierID, bool bIsOn)
+	{
+		return OnUpdateModifierStateFunc(ModifierID, bIsOn);
+	}
 };
