UVEditor: Add comment about using perspective instead of ortho projection for 2d viewport.

diff --git a/Engine/Plugins/Experimental/UVEditor/Source/UVEditor/Private/UVEditor2DViewportClient.cpp b/Engine/Plugins/Experimental/UVEditor/Source/UVEditor/Private/UVEditor2DViewportClient.cpp
--- a/Engine/Plugins/Experimental/UVEditor/Source/UVEditor/Private/UVEditor2DViewportClient.cpp
+++ b/Engine/Plugins/Experimental/UVEditor/Source/UVEditor/Private/UVEditor2DViewportClient.cpp
@@ -160,6 +160,8 @@ FUVEditor2DViewportClient::FUVEditor2DViewportClient(FEditorModeTools* InModeToo
 	OverrideNearClipPlane(KINDA_SMALL_NUMBER);
 
 	// Set up viewport manipulation behaviors:
+	// Note that this is only necessary because we use a perspective projection viewport 
+	// instead of a proper ortho viewport. See comment in UVEditorToolkit.cpp concerning this choice.
 
 	FEditorCameraController* CameraControllerPtr = GetCameraController();
 	CameraController->GetConfig().MovementAccelerationRate = 0.0;
diff --git a/Engine/Plugins/Experimental/UVEditor/Source/UVEditor/Private/UVEditorToolkit.cpp b/Engine/Plugins/Experimental/UVEditor/Source/UVEditor/Private/UVEditorToolkit.cpp
--- a/Engine/Plugins/Experimental/UVEditor/Source/UVEditor/Private/UVEditorToolkit.cpp
+++ b/Engine/Plugins/Experimental/UVEditor/Source/UVEditor/Private/UVEditorToolkit.cpp
@@ -599,7 +599,20 @@ void FUVEditorToolkit::PostInitAssetEditor()
 	};
 	SetCommonViewportClientOptions(ViewportClient.Get());
 
-	// Ortho has too many problems with rendering things, unfortunately, so we should use perspective.
+	// Note about using perspective instead of ortho projection for a 2d viewport:
+	// Originally, we had various issues with the materials that we were using in ortho mode. Most of
+	// those issues seem to have been resolved since then, and an example of using a proper 2d viewport
+	// can be seen in ClothEditorRestSpaceViewportClient.cpp and ClothEditorToolkit.cpp. The benefits of
+	// doing so are viewport controls out of the box, and ability to switch the viewport type at will.
+	// 
+	// At this point, however, the UV editor would not get much benefit from converting to a true ortho
+	// viewport, and the conversion carries with it some amount of risk. At time of last experimentation 
+	// with the conversion (1/3/2023), the materials seemed to all work, but the gizmo failed to draw
+	// (unknown reason), and the grid number labels needed fixing.
+	// 
+	// Until we run into issues, we will continue to use perspective projection to avoid the conversion
+	// cost/risks and to keep the advantages of more trustworthy rendering, but if you use the editor as
+	// a model, consider using an ortho viewport for the potential simplicity.
 	ViewportClient->SetViewportType(ELevelViewportType::LVT_Perspective);
 
 	// Lit gives us the most options in terms of the materials we can use.
