UVEditor: Fix gizmo no longer showing up after the changes that made gizmo respond to editor gizmo modes.

diff --git a/Engine/Plugins/Experimental/UVEditor/Source/UVEditor/Private/UVEditorMode.cpp b/Engine/Plugins/Experimental/UVEditor/Source/UVEditor/Private/UVEditorMode.cpp
--- a/Engine/Plugins/Experimental/UVEditor/Source/UVEditor/Private/UVEditorMode.cpp
+++ b/Engine/Plugins/Experimental/UVEditor/Source/UVEditor/Private/UVEditorMode.cpp
@@ -382,6 +382,11 @@ void UUVEditorMode::Enter()
 	InitializeModeContexts();
 	InitializeTargets();
 
+	// By default combined transform gizmos now hide some of their elements based on what gizmo mode is chosen in
+	// the viewport. The UV editor doesn't currently have that selector, so we don't want to try to use it. If we
+	// add that selector, we'll need to remove this call.
+	GetInteractiveToolsContext()->SetForceCombinedGizmoMode(true);
+
 	BackgroundVisualization = NewObject<UUVEditorBackgroundPreview>(this);
 	BackgroundVisualization->CreateInWorld(GetWorld(), FTransform::Identity);
 
