UVEditor: Change mapping from uv to world to make the X/Y gizmo axis colors align with U/V. Add helpful mapping functions and clarify their usage in comments. Use the new helper functions where relevant.

diff --git a/Engine/Plugins/Experimental/UVEditor/Source/UVEditor/Private/UVEditor2DViewportClient.cpp b/Engine/Plugins/Experimental/UVEditor/Source/UVEditor/Private/UVEditor2DViewportClient.cpp
--- a/Engine/Plugins/Experimental/UVEditor/Source/UVEditor/Private/UVEditor2DViewportClient.cpp
+++ b/Engine/Plugins/Experimental/UVEditor/Source/UVEditor/Private/UVEditor2DViewportClient.cpp
@@ -125,9 +125,9 @@ namespace FUVEditor2DViewportClientLocals {
 
 	bool ConvertUVToPixel(const FVector2D& UVIn, FVector2D& PixelOut, const FSceneView& View)
 	{
-		FVector TestWorld = FUVEditorUXSettings::UVToVertPosition(FUVEditorUXSettings::ExternalUVToInternalUV((FVector2f)UVIn));
-		FVector4 TestProjectedHomogenous = View.WorldToScreen(TestWorld);
-		bool bValid = View.ScreenToPixel(TestProjectedHomogenous, PixelOut);
+		FVector WorldPosition = FUVEditorUXSettings::ExternalUVToUnwrapWorldPosition((FVector2f)UVIn);
+		FVector4 ProjectedHomogenous = View.WorldToScreen(WorldPosition);
+		bool bValid = View.ScreenToPixel(ProjectedHomogenous, PixelOut);
 		return bValid;
 	}
 
@@ -138,8 +138,7 @@ namespace FUVEditor2DViewportClientLocals {
 		FVector WorldPoint(WorldPointHomogenous.X / WorldPointHomogenous.W,
 			               WorldPointHomogenous.Y / WorldPointHomogenous.W,
 			               WorldPointHomogenous.Z / WorldPointHomogenous.W);
-		UVOut = (FVector2D)FUVEditorUXSettings::VertPositionToUV(WorldPoint);
-		UVOut = (FVector2D)FUVEditorUXSettings::InternalUVToExternalUV((FVector2f)UVOut);
+		UVOut = (FVector2D)FUVEditorUXSettings::UnwrapWorldPositionToExternalUV(WorldPoint);
 	}
 
 };
@@ -287,11 +286,13 @@ void FUVEditor2DViewportClient::DrawGrid(const FSceneView* View, FPrimitiveDrawI
 			PDI, Transform);
 	}
 
-	double AxisExtent = UVScale;
-
 	// Draw colored axis lines
-	PDI->DrawLine(FVector(0, 0, 0), FVector(AxisExtent, 0, 0), FUVEditorUXSettings::XAxisColor, SDPG_World, FUVEditorUXSettings::AxisThickness, 0, true);
-	PDI->DrawLine(FVector(0, 0, 0), FVector(0, AxisExtent, 0), FUVEditorUXSettings::YAxisColor, SDPG_World, FUVEditorUXSettings::AxisThickness, 0, true);
+	PDI->DrawLine(FUVEditorUXSettings::ExternalUVToUnwrapWorldPosition(FVector2f(0,0)), 
+		FUVEditorUXSettings::ExternalUVToUnwrapWorldPosition(FVector2f(1, 0)), 
+		FUVEditorUXSettings::XAxisColor, SDPG_World, FUVEditorUXSettings::AxisThickness, 0, true);
+	PDI->DrawLine(FUVEditorUXSettings::ExternalUVToUnwrapWorldPosition(FVector2f(0, 0)), 
+		FUVEditorUXSettings::ExternalUVToUnwrapWorldPosition(FVector2f(0, 1)), 
+		FUVEditorUXSettings::YAxisColor, SDPG_World, FUVEditorUXSettings::AxisThickness, 0, true);
 
 	// TODO: Draw a little UV axis thing in the lower left, like the XYZ things that normal viewports have.
 }
diff --git a/Engine/Plugins/Experimental/UVEditor/Source/UVEditor/Private/UVEditorBackgroundPreview.cpp b/Engine/Plugins/Experimental/UVEditor/Source/UVEditor/Private/UVEditorBackgroundPreview.cpp
--- a/Engine/Plugins/Experimental/UVEditor/Source/UVEditor/Private/UVEditorBackgroundPreview.cpp
+++ b/Engine/Plugins/Experimental/UVEditor/Source/UVEditor/Private/UVEditorBackgroundPreview.cpp
@@ -187,47 +187,40 @@ void UUVEditorBackgroundPreview::UpdateBackground()
 	BackgroundMaterial->SetScalarParameterValue(TEXT("BackgroundPixelDepthOffset"), FUVEditorUXSettings::BackgroundQuadDepthOffset);	
 	BackgroundComponent->Clear();
 
-	TArray<FVector2f> UDimBlocksToRender;
+	TArray<FIndex2i> UDimBlocksToRender;
 	if (Settings->bUDIMsEnabled)
 	{
 		// TODO: Find a way to access the list of UDIMs from the context object instead? 
 		for (int32 BlockIndex = 0; BlockIndex < Settings->UDIMBlocks.Num(); ++BlockIndex)
 		{
-			FVector2f Block;
-			int32 UCoord, VCoord;
-			UE::TextureUtilitiesCommon::ExtractUDIMCoordinates(Settings->UDIMBlocks[BlockIndex], UCoord, VCoord);
-			Block.X = static_cast<float>(UCoord);
-			Block.Y = static_cast<float>(VCoord);
+			FIndex2i Block;
+			UE::TextureUtilitiesCommon::ExtractUDIMCoordinates(Settings->UDIMBlocks[BlockIndex], Block.A, Block.B);
 			UDimBlocksToRender.Push(Block);
 		}
 	}
 	if (UDimBlocksToRender.Num() == 0)
 	{
-		UDimBlocksToRender.Push(FVector2f(0, 0));
+		UDimBlocksToRender.Push(FIndex2i(0, 0));
 	}
 
-	for (const FVector2f& GridStep : UDimBlocksToRender)
+	for (const FIndex2i& Block : UDimBlocksToRender)
 	{
-		FVector2f UV00 = { (GridStep.X + 0.0f) , (GridStep.Y + 0.0f) };
-		FVector2f UV01 = { (GridStep.X + 1.0f) , (GridStep.Y + 0.0f) };		
-		FVector2f UV10 = { (GridStep.X + 0.0f) , (GridStep.Y + 1.0f) };
-		FVector2f UV11 = { (GridStep.X + 1.0f) , (GridStep.Y + 1.0f) };
-
-		UV00 = FUVEditorUXSettings::ExternalUVToInternalUV(UV00);
-		UV01 = FUVEditorUXSettings::ExternalUVToInternalUV(UV01);
-		UV10 = FUVEditorUXSettings::ExternalUVToInternalUV(UV10);
-		UV11 = FUVEditorUXSettings::ExternalUVToInternalUV(UV11);
-
-		FRenderableTriangleVertex A((FVector)FUVEditorUXSettings::UVToVertPosition(UV00), (FVector2D)UV00, Normal, BackgroundColor);
-		FRenderableTriangleVertex B((FVector)FUVEditorUXSettings::UVToVertPosition(UV10), (FVector2D)UV10, Normal, BackgroundColor);
-		FRenderableTriangleVertex C((FVector)FUVEditorUXSettings::UVToVertPosition(UV01), (FVector2D)UV01, Normal, BackgroundColor);
-		FRenderableTriangleVertex D((FVector)FUVEditorUXSettings::UVToVertPosition(UV11), (FVector2D)UV11, Normal, BackgroundColor);
-
-		FRenderableTriangle Lower(BackgroundMaterial, A, D, B);
-		FRenderableTriangle Upper(BackgroundMaterial, A, C, D);
-
-		BackgroundComponent->AddTriangle(Lower);
-		BackgroundComponent->AddTriangle(Upper);
+		auto MakeQuadVert = [&Block, &Normal, &BackgroundColor](int32 CornerX, int32 CornerY)
+		{
+			FVector2f ExternalUV(static_cast<float>(Block.A + CornerX), static_cast<float>(Block.B + CornerY));
+
+			return FRenderableTriangleVertex(FUVEditorUXSettings::ExternalUVToUnwrapWorldPosition(ExternalUV),
+				(FVector2D)FUVEditorUXSettings::ExternalUVToInternalUV(ExternalUV),
+				Normal, BackgroundColor);
+		};
+
+		FRenderableTriangleVertex V00 = MakeQuadVert(0,0);
+		FRenderableTriangleVertex V10 = MakeQuadVert(1, 0);
+		FRenderableTriangleVertex V11 = MakeQuadVert(1, 1);
+		FRenderableTriangleVertex V01 = MakeQuadVert(0,1);
+
+		BackgroundComponent->AddTriangle(FRenderableTriangle(BackgroundMaterial, V00, V10, V11));
+		BackgroundComponent->AddTriangle(FRenderableTriangle(BackgroundMaterial, V00, V11, V01));
 	}
 }
 
diff --git a/Engine/Plugins/Experimental/UVEditor/Source/UVEditor/Private/UVEditorModeToolkit.cpp b/Engine/Plugins/Experimental/UVEditor/Source/UVEditor/Private/UVEditorModeToolkit.cpp
--- a/Engine/Plugins/Experimental/UVEditor/Source/UVEditor/Private/UVEditorModeToolkit.cpp
+++ b/Engine/Plugins/Experimental/UVEditor/Source/UVEditor/Private/UVEditorModeToolkit.cpp
@@ -208,12 +208,12 @@ void FUVEditorModeToolkit::Init(const TSharedPtr<IToolkitHost>& InitToolkitHost,
 		// Make our displays show 0-1 units instead of world units.
 		.InternalToDisplayFunction([](const FVector3d& WorldPosition)
 		{
-			FVector2f UV = FUVEditorUXSettings::InternalUVToExternalUV(FUVEditorUXSettings::VertPositionToUV(WorldPosition));
+			FVector2f UV = FUVEditorUXSettings::UnwrapWorldPositionToExternalUV(WorldPosition);
 			return FVector3d(UV.X, UV.Y, 0);
 		})
 		.DisplayToInternalFunction([](const FVector3d& UVPosition)
 		{
-			return FUVEditorUXSettings::UVToVertPosition(FUVEditorUXSettings::ExternalUVToInternalUV(FVector2f(UVPosition.X, UVPosition.Y)));
+			return FUVEditorUXSettings::ExternalUVToUnwrapWorldPosition(FVector2f(UVPosition.X, UVPosition.Y));
 		})
 		.TranslationScrubSensitivity(1 / FUVEditorUXSettings::UVMeshScalingFactor)
 		;
diff --git a/Engine/Plugins/Experimental/UVEditor/Source/UVEditor/Private/UVEditorToolkit.cpp b/Engine/Plugins/Experimental/UVEditor/Source/UVEditor/Private/UVEditorToolkit.cpp
--- a/Engine/Plugins/Experimental/UVEditor/Source/UVEditor/Private/UVEditorToolkit.cpp
+++ b/Engine/Plugins/Experimental/UVEditor/Source/UVEditor/Private/UVEditorToolkit.cpp
@@ -519,7 +519,7 @@ void FUVEditorToolkit::PostInitAssetEditor()
 			MenuBuilder.BeginSection("Section_Gizmo", LOCTEXT("Section_Gizmo", "Gizmo"));
 
 			MenuBuilder.AddSubMenu(
-				LOCTEXT("NumericalUISubMenu", "Numerical UI"), LOCTEXT("NumericalUISubMenu_ToolTip", "Configure the gizmo numerical UI."),
+				LOCTEXT("GizmoTransformPanelSubMenu", "Gizmo Transform Panel"), LOCTEXT("GizmoTransformPanelSubMenu_ToolTip", "Configure the Gizmo Transform Panel."),
 				FNewMenuDelegate::CreateLambda([UVModeToolkit](FMenuBuilder& SubMenuBuilder) {
 					UVModeToolkit->MakeGizmoNumericalUISubMenu(SubMenuBuilder);
 				}));
@@ -644,14 +644,18 @@ void FUVEditorToolkit::PostInitAssetEditor()
 	// scale [0,1] to [0,ScaleFactor].
 	// We set our camera to look downward, centered, far enough to be able to see the edges
 	// with a 90 degree FOV
-	double ScaleFactor = 1;
-	UUVEditorSubsystem* UVSubsystem = GEditor->GetEditorSubsystem<UUVEditorSubsystem>();
-	if (UVSubsystem)
-	{
-		ScaleFactor = FUVEditorUXSettings::UVMeshScalingFactor;
-	}
-	ViewportClient->SetViewLocation(FVector(ScaleFactor / 2, ScaleFactor / 2, ScaleFactor));
-	ViewportClient->SetViewRotation(FRotator(-90, 0, 0));
+	FVector3d ViewLocation = FUVEditorUXSettings::ExternalUVToUnwrapWorldPosition(FVector2f(0.5, 0.5));
+	ViewLocation.Z = FUVEditorUXSettings::UVMeshScalingFactor;
+	ViewportClient->SetViewLocation(ViewLocation);
+	
+	// Note that the orientation at which we look at the world should match to the transformations
+	// we set in FUVEditorUXSettings::ExternalUVToUnwrapWorldPosition, etc.
+	// Currently we choose to view the unwrap world with Z pointing at us, X pointing right, and 
+	// Y pointing down, which matches the view given by LVT_OrthoXY, which we may someday decide to use.
+	// Also, we want the U and V axes to lie along the X and Y axes, respectively, even if they are
+	// flipped, because that makes the colors used in the gizmo and its numerical UI panel match up 
+	// with what the user expects.
+	ViewportClient->SetViewRotation(FRotator(-90, 0, -90));
 
 	// If exposure isn't set to fixed, it will flash as we stare into the void
 	ViewportClient->ExposureSettings.bFixed = true;
diff --git a/Engine/Plugins/Experimental/UVEditor/Source/UVEditorTools/Private/Operators/UVEditorUVTransformOp.cpp b/Engine/Plugins/Experimental/UVEditor/Source/UVEditorTools/Private/Operators/UVEditorUVTransformOp.cpp
--- a/Engine/Plugins/Experimental/UVEditor/Source/UVEditorTools/Private/Operators/UVEditorUVTransformOp.cpp
+++ b/Engine/Plugins/Experimental/UVEditor/Source/UVEditorTools/Private/Operators/UVEditorUVTransformOp.cpp
@@ -16,15 +16,7 @@ using namespace UE::Geometry;
 
 namespace TransformOpLocals
 {
-	FVector2f UnwrapPositionToUV(const FVector3d& Position)
-	{
-		return FUVEditorUXSettings::ExternalUVToInternalUV(FUVEditorUXSettings::VertPositionToUV(Position));
-	}
 
-	FVector3d UVToUnwrapPosition(const FVector2f& UV)
-	{
-		return FUVEditorUXSettings::UVToVertPosition(FUVEditorUXSettings::InternalUVToExternalUV(UV));
-	}
 }
 
 void FUVEditorUVTransformBaseOp::SetTransform(const FTransformSRT3d& Transform)
@@ -82,7 +74,7 @@ void FUVEditorUVTransformBaseOp::RebuildBoundingBoxes()
 			const TArray<int>& Vertices = (*UVComponents)[k].Indices;
 			for (int32 Vid : Vertices)
 			{				
-				PerComponentBoundingBoxes[k].Contain((FVector2d)TransformOpLocals::UnwrapPositionToUV(ResultMesh->GetVertexRef(Vid)));
+				PerComponentBoundingBoxes[k].Contain((FVector2d)FUVEditorUXSettings::UnwrapWorldPositionToExternalUV(ResultMesh->GetVertexRef(Vid)));
 			}
 		});
 
@@ -306,17 +298,17 @@ void FUVEditorUVTransformOp::HandleTransformationOp(FProgressCancel * Progress)
 	auto ScaleFunc = [this](int32 Vid)
 	{
 		FVector2f ScalePivot = GetPivotFromMode(Vid, PivotMode);
-		FVector2f UV = TransformOpLocals::UnwrapPositionToUV(ResultMesh->GetVertexRef(Vid));
+		FVector2f UV = FUVEditorUXSettings::UnwrapWorldPositionToExternalUV(ResultMesh->GetVertexRef(Vid));
 		UV = (UV - ScalePivot);
 		UV[0] *= static_cast<float>(Scale[0]);
 		UV[1] *= static_cast<float>(Scale[1]);
 		UV = (UV + ScalePivot);
-		ResultMesh->SetVertex(Vid, TransformOpLocals::UVToUnwrapPosition(UV));
+		ResultMesh->SetVertex(Vid, FUVEditorUXSettings::ExternalUVToUnwrapWorldPosition(UV));
 	};
 
 	auto BaseRotFunc = [this](int32 Vid, float RotationIn, const FVector2f& Pivot)
 	{
-		FVector2f UV = TransformOpLocals::UnwrapPositionToUV(ResultMesh->GetVertexRef(Vid));
+		FVector2f UV = FUVEditorUXSettings::UnwrapWorldPositionToExternalUV(ResultMesh->GetVertexRef(Vid));
 		FVector2f UV_Rotated;
 		// We are flipping the sign here to match the conventions in other UV editors
 		// where positive values are clockwise and negative values are counterclockwise.
@@ -325,7 +317,7 @@ void FUVEditorUVTransformOp::HandleTransformationOp(FProgressCancel * Progress)
 		UV_Rotated[0] = UV[0] * static_cast<float>(FMath::Cos(RotationInRadians)) - UV[1] * static_cast<float>(FMath::Sin(RotationInRadians));
 		UV_Rotated[1] = UV[0] * static_cast<float>(FMath::Sin(RotationInRadians)) + UV[1] * static_cast<float>(FMath::Cos(RotationInRadians));
 		UV = (UV_Rotated + Pivot);
-		ResultMesh->SetVertex(Vid, TransformOpLocals::UVToUnwrapPosition(UV));
+		ResultMesh->SetVertex(Vid, FUVEditorUXSettings::ExternalUVToUnwrapWorldPosition(UV));
 	};
 
 	auto QuickRotFunc = [this, &BaseRotFunc](int32 Vid)
@@ -342,9 +334,9 @@ void FUVEditorUVTransformOp::HandleTransformationOp(FProgressCancel * Progress)
 
 	auto QuickTranslateFunc = [this](int32 Vid)
 	{
-		FVector2f UV = TransformOpLocals::UnwrapPositionToUV(ResultMesh->GetVertexRef(Vid));
+		FVector2f UV = FUVEditorUXSettings::UnwrapWorldPositionToExternalUV(ResultMesh->GetVertexRef(Vid));
 		UV = (UV + FVector2f(QuickTranslation));
-		ResultMesh->SetVertex(Vid, TransformOpLocals::UVToUnwrapPosition(UV));
+		ResultMesh->SetVertex(Vid, FUVEditorUXSettings::ExternalUVToUnwrapWorldPosition(UV));
 	};
 
 	auto TranslateFunc = [this](int32 Vid)
@@ -354,9 +346,9 @@ void FUVEditorUVTransformOp::HandleTransformationOp(FProgressCancel * Progress)
 		{
 			RotationPivot = GetPivotFromMode(Vid, PivotMode);
 		}				
-		FVector2f UV = TransformOpLocals::UnwrapPositionToUV(ResultMesh->GetVertexRef(Vid));
+		FVector2f UV = FUVEditorUXSettings::UnwrapWorldPositionToExternalUV(ResultMesh->GetVertexRef(Vid));
 		UV = (UV + FVector2f(Translation) - RotationPivot);
-		ResultMesh->SetVertex(Vid, TransformOpLocals::UVToUnwrapPosition(UV));
+		ResultMesh->SetVertex(Vid, FUVEditorUXSettings::ExternalUVToUnwrapWorldPosition(UV));
 	};
 
 	auto ApplyTransformFunc = [this, &Progress](TFunction<void(int32 Vid)> TransformFunc)
@@ -446,9 +438,9 @@ void FUVEditorUVAlignOp::HandleTransformationOp(FProgressCancel* Progress)
 
 	auto TranslateFunc = [this](const FVector2f& Translation, int32 Vid)
 	{
-		FVector2f UV = TransformOpLocals::UnwrapPositionToUV(ResultMesh->GetVertexRef(Vid));
+		FVector2f UV = FUVEditorUXSettings::UnwrapWorldPositionToExternalUV(ResultMesh->GetVertexRef(Vid));
 		UV = (UV + Translation);
-		ResultMesh->SetVertex(Vid, TransformOpLocals::UVToUnwrapPosition(UV));
+		ResultMesh->SetVertex(Vid, FUVEditorUXSettings::ExternalUVToUnwrapWorldPosition(UV));
 	};
 
 	auto TranslationFromAlignmentPoints = [this](const FVector2f& PointTo, const FVector2f& PointFrom)
@@ -481,7 +473,7 @@ void FUVEditorUVAlignOp::HandleTransformationOp(FProgressCancel* Progress)
 	};
 	auto UVElementToUDIM = [this](int32 Vid)
 	{
-		return FDynamicMeshUDIMClassifier::ClassifyPointToUDIM(TransformOpLocals::UnwrapPositionToUV(ResultMesh->GetVertexRef(Vid)));
+		return FDynamicMeshUDIMClassifier::ClassifyPointToUDIM(FUVEditorUXSettings::UnwrapWorldPositionToExternalUV(ResultMesh->GetVertexRef(Vid)));
 	};
 	auto ComponentToPoint = [this](int32 ComponentID)
 	{
@@ -489,7 +481,7 @@ void FUVEditorUVAlignOp::HandleTransformationOp(FProgressCancel* Progress)
 	};
 	auto UVElementToPoint = [this](int32 Vid)
 	{
-		return TransformOpLocals::UnwrapPositionToUV(ResultMesh->GetVertexRef(Vid));
+		return FUVEditorUXSettings::UnwrapWorldPositionToExternalUV(ResultMesh->GetVertexRef(Vid));
 	};
 
 
@@ -757,9 +749,9 @@ void FUVEditorUVDistributeOp::HandleTransformationOp(FProgressCancel* Progress)
 
 	auto TranslateFunc = [this](const FVector2f& Translation, int32 Vid)
 	{
-		FVector2f UV = TransformOpLocals::UnwrapPositionToUV(ResultMesh->GetVertexRef(Vid));
+		FVector2f UV = FUVEditorUXSettings::UnwrapWorldPositionToExternalUV(ResultMesh->GetVertexRef(Vid));
 		UV = (UV + Translation);
-		ResultMesh->SetVertex(Vid, TransformOpLocals::UVToUnwrapPosition(UV));
+		ResultMesh->SetVertex(Vid, FUVEditorUXSettings::ExternalUVToUnwrapWorldPosition(UV));
 	};
 
 	if (TransformingElements.IsSet())
diff --git a/Engine/Plugins/Experimental/UVEditor/Source/UVEditorTools/Private/UVEditorTransformTool.cpp b/Engine/Plugins/Experimental/UVEditor/Source/UVEditorTools/Private/UVEditorTransformTool.cpp
--- a/Engine/Plugins/Experimental/UVEditor/Source/UVEditorTools/Private/UVEditorTransformTool.cpp
+++ b/Engine/Plugins/Experimental/UVEditor/Source/UVEditorTools/Private/UVEditorTransformTool.cpp
@@ -366,7 +366,7 @@ void UUVEditorTransformTool::DrawHUD(FCanvas* Canvas, IToolsContextRenderAPI* Re
 
 	auto ConvertUVToPixel = [RenderAPI](const FVector2D& UVIn, FVector2D& PixelOut)
 	{
-		FVector WorldPoint = FUVEditorUXSettings::UVToVertPosition(FUVEditorUXSettings::ExternalUVToInternalUV((FVector2f)UVIn));
+		FVector WorldPoint = FUVEditorUXSettings::ExternalUVToUnwrapWorldPosition((FVector2f)UVIn);
 		FVector4 TestProjectedHomogenous = RenderAPI->GetSceneView()->WorldToScreen(WorldPoint);
 		bool bValid = RenderAPI->GetSceneView()->ScreenToPixel(TestProjectedHomogenous, PixelOut);
 		return bValid;
diff --git a/Engine/Plugins/Experimental/UVEditor/Source/UVEditorTools/Private/UVEditorUXSettings.cpp b/Engine/Plugins/Experimental/UVEditor/Source/UVEditorTools/Private/UVEditorUXSettings.cpp
--- a/Engine/Plugins/Experimental/UVEditor/Source/UVEditorTools/Private/UVEditorUXSettings.cpp
+++ b/Engine/Plugins/Experimental/UVEditor/Source/UVEditorTools/Private/UVEditorUXSettings.cpp
@@ -128,26 +128,44 @@ FLinearColor FUVEditorUXSettings::GetBoundaryColorByTargetIndex(int32 TargetInde
 	return BoundaryColorHSV.HSVToLinearRGB();
 }
 
+// The mappings here depend on the way we view the unwrap world, set up in UVEditorToolkit.cpp.
+// Currently, we look at the world with Z pointing towards the camera, the positive X axis pointing right,
+// and the positive Y axis pointing down. This means that we have to flip the V coordinate to map it
+// to the up direction.
+FVector3d FUVEditorUXSettings::ExternalUVToUnwrapWorldPosition(const FVector2f& UV)
+{
+	return FVector3d(UV.X * UVMeshScalingFactor, -UV.Y * UVMeshScalingFactor, 0);
+}
+FVector2f FUVEditorUXSettings::UnwrapWorldPositionToExternalUV(const FVector3d& VertPosition)
+{
+	return FVector2f(static_cast<float>(VertPosition.X) / UVMeshScalingFactor, -(static_cast<float>(VertPosition.Y) / UVMeshScalingFactor));
+}
+
+// Unreal stores its UVs with V subtracted from 1.
 FVector2f FUVEditorUXSettings::ExternalUVToInternalUV(const FVector2f& UV)
 {
 	return FVector2f(UV.X, 1-UV.Y);
 }
-
 FVector2f FUVEditorUXSettings::InternalUVToExternalUV(const FVector2f& UV)
 {
 	return FVector2f(UV.X, 1-UV.Y);
 }
 
+// Should be equivalent to converting from unreal's UV to regular (external) and then to unwrap world,
+// i.e. ExternalUVToUnwrapWorldPosition(InternalUVToExternalUV(UV))
 FVector3d FUVEditorUXSettings::UVToVertPosition(const FVector2f& UV)
 {
-	return FVector3d((1 - UV.Y) * UVMeshScalingFactor, UV.X * UVMeshScalingFactor, 0);
+	return FVector3d(UV.X * UVMeshScalingFactor, (UV.Y - 1) * UVMeshScalingFactor, 0);
 }
 
+// Should be equivalent to converting from world to regular (external) UV, and then to unreal's representation,
+// i.e. ExternalUVToInternalUV(UnwrapWorldPositionToExternalUV(VertPosition))
 FVector2f FUVEditorUXSettings::VertPositionToUV(const FVector3d& VertPosition)
 {
-	return FVector2f(static_cast<float>(VertPosition.Y) / UVMeshScalingFactor, 1.0f - (static_cast<float>(VertPosition.X) / UVMeshScalingFactor));
+	return FVector2f(static_cast<float>(VertPosition.X) / UVMeshScalingFactor, 1.0 + (static_cast<float>(VertPosition.Y) / UVMeshScalingFactor));
 };
 
+
 float FUVEditorUXSettings::LocationSnapValue(int32 LocationSnapMenuIndex)
 {
 	switch (LocationSnapMenuIndex)
diff --git a/Engine/Plugins/Experimental/UVEditor/Source/UVEditorTools/Public/UVEditorUXSettings.h b/Engine/Plugins/Experimental/UVEditor/Source/UVEditorTools/Public/UVEditorUXSettings.h
--- a/Engine/Plugins/Experimental/UVEditor/Source/UVEditorTools/Public/UVEditorUXSettings.h
+++ b/Engine/Plugins/Experimental/UVEditor/Source/UVEditorTools/Public/UVEditorUXSettings.h
@@ -113,26 +113,55 @@ public:
 	//---------------------------
 
 
-	// When creating UV unwraps or transforming between UV coordinate and positions in the 2D viewport,
-	// these functions will determine the mapping between UV values and the
-	// resulting unwrap mesh vertex positions. 
-	// If we're looking down on the unwrapped mesh, with the Z axis towards us, we want U's to be right, and
-	// V's to be up. In Unreal's left-handed coordinate system, this means that we map U's to world Y
-	// and V's to world X.
-	// Also, Unreal changes the V coordinates of imported meshes to 1-V internally, and we undo this
-	// while displaying the UV's because the users likely expect to see the original UV's (it would
-	// be particularly confusing for users working with UDIM assets, where internally stored V's 
-	// frequently end up negative).
-	// The ScaleFactor just scales the mesh up. Scaling the mesh up makes it easier to zoom in
-	// further into the display before getting issues with the camera near plane distance.
+	// Note about the conversions from unwrap world to UV, below:
+	// Nothing should make assumptions about the mapping between world space and UV coordinates. Instead,
+	// tools should use the conversion functions below so that the mapping is free to change without affecting 
+	// the tools.
+
+	/**
+	 * Converts from UV value as displayed to user or seen before import, to point in unwrap world space. This
+	 * is useful for mapping visualizations to world unwrap space.
+	 *
+	 * Like other world-UV conversion functions, clients should not know about the details of this conversion.
+	 */
+	static FVector3d ExternalUVToUnwrapWorldPosition(const FVector2f& UV);
+
+	/**
+	 * Converts from point in unwrap world space, to UV value as displayed to user or seen before import. This
+	 * is useful for labeling UV values in the unwrap world, even though the UV values are stored slightly differently
+	 * inside the mesh itself.
+	 *
+	 * Like other world-UV conversion functions, clients should not know about the details of this conversion.
+	 */
+	static FVector2f UnwrapWorldPositionToExternalUV(const FVector3d& VertPosition);
+
+	/**
+	 * Converts from UV value as displayed to user or seen before import, to UV as stored on a mesh in Unreal.
+	 */
 	static FVector2f ExternalUVToInternalUV(const FVector2f& UV);
 
+	/**
+	 * Converts from UV as stored on a mesh in Unreal to UV as displayed to user or seen in an external program.
+	 */
 	static FVector2f InternalUVToExternalUV(const FVector2f& UV);
 
+	/**
+	 * Convert from UV as stored on a mesh in Unreal to world position in the UV editor unwrap world. This allows
+	 * changes in UVs of a mesh to be mapped to its unwrap representation.
+	 * 
+	 * Like other world-UV conversion functions, clients should not know about the details of this conversion.
+	 */
 	static FVector3d UVToVertPosition(const FVector2f& UV);
 
+	/**
+	 * Converts from position in UV editor unwrap world to UV value as stored on a mesh in Unreal. This allows
+	 * changes in the unwrap world to be mapped to actual mesh UVs.
+	 * 
+	 * Like other world-UV conversion functions, clients should not know about the details of this conversion.
+	 */
 	static FVector2f VertPositionToUV(const FVector3d& VertPosition);
 
+
 	//--------------------------------
 	// CVARs for Experimental Features
 	//--------------------------------
