UV Editor: Change the quick translate options in the transform tool to produce a single undo entry for the corner transforms.

diff --git a/Engine/Plugins/Experimental/UVEditor/Source/UVEditorToolsEditorOnly/Private/DetailsCustomizations/UVTransformToolCustomizations.cpp b/Engine/Plugins/Experimental/UVEditor/Source/UVEditorToolsEditorOnly/Private/DetailsCustomizations/UVTransformToolCustomizations.cpp
--- a/Engine/Plugins/Experimental/UVEditor/Source/UVEditorToolsEditorOnly/Private/DetailsCustomizations/UVTransformToolCustomizations.cpp
+++ b/Engine/Plugins/Experimental/UVEditor/Source/UVEditorToolsEditorOnly/Private/DetailsCustomizations/UVTransformToolCustomizations.cpp
@@ -16,6 +16,8 @@
 #include "IContentBrowserSingleton.h"
 #include "SAssetView.h"
 
+#include "Editor.h" // For supporting Undo transactions in place
+
 #include "PropertyRestriction.h"
 
 #include "UVEditorStyle.h"
@@ -133,13 +135,15 @@ void FUVEditorUVTransformToolDetails::BuildQuickTranslateMenu(IDetailLayoutBuild
 	ensure(QuickTranslateHandle->IsValidHandle());
 	QuickTranslateHandle->MarkHiddenByCustomization();
 
-	auto ApplyTranslation = [this, QuickTranslateHandle, QuickTranslateOffsetHandle](const FVector2D& Direction)
+	auto ApplyTranslation = [this, QuickTranslateHandle, QuickTranslateOffsetHandle](FText TransactionDescription, const FVector2D& Direction)
 	{
+		GEditor->BeginTransaction(TransactionDescription);
 		float TranslationValue;
 		FVector2D TotalTranslationValue;
 		QuickTranslateOffsetHandle->GetValue(TranslationValue);
 		QuickTranslateHandle->GetValue(TotalTranslationValue);
 		QuickTranslateHandle->SetValue(TotalTranslationValue + (Direction * TranslationValue) );
+		GEditor->EndTransaction();
 		return FReply::Handled();
 	};
 
@@ -154,7 +158,7 @@ void FUVEditorUVTransformToolDetails::BuildQuickTranslateMenu(IDetailLayoutBuild
 			SNew(SButton)
 			.Text(LOCTEXT("QuickMoveTopLeft", "TL"))
 			.ToolTipText(LOCTEXT("QuickMoveTopLeftToolTip", "Applies translation offset in the negative X axis and the positive Y axis"))
-			.OnClicked_Lambda([ApplyTranslation]() { return ApplyTranslation(FVector2D(-1.0, 1.0)); })		
+			.OnClicked_Lambda([ApplyTranslation]() { return ApplyTranslation(LOCTEXT("TransactionTopLeft", "Top Left Translation"), FVector2D(-1.0, 1.0)); })
 		]
 		+ SUniformGridPanel::Slot(1, 0)
 		.HAlign(HAlign_Center)
@@ -163,7 +167,7 @@ void FUVEditorUVTransformToolDetails::BuildQuickTranslateMenu(IDetailLayoutBuild
 			SNew(SButton)
 			.Text(LOCTEXT("QuickMoveTop", "T"))
 			.ToolTipText(LOCTEXT("QuickMoveTopToolTip", "Applies the translation offset in the positive Y axis."))
-			.OnClicked_Lambda([ApplyTranslation]() { return ApplyTranslation(FVector2D(0.0, 1.0)); })
+			.OnClicked_Lambda([ApplyTranslation]() { return ApplyTranslation(LOCTEXT("TransactionTop", "Top Translation"), FVector2D(0.0, 1.0)); })
 		]
 		+ SUniformGridPanel::Slot(2, 0)
 		.HAlign(HAlign_Center)
@@ -172,7 +176,7 @@ void FUVEditorUVTransformToolDetails::BuildQuickTranslateMenu(IDetailLayoutBuild
 			SNew(SButton)
 			.Text(LOCTEXT("QuickMoveTopRight", "TR"))
 			.ToolTipText(LOCTEXT("QuickMoveTopRightToolTip", "Applies the translation offset in the positive X axis and the positive Y axis."))
-			.OnClicked_Lambda([ApplyTranslation]() { return ApplyTranslation(FVector2D(1.0, 1.0)); })
+			.OnClicked_Lambda([ApplyTranslation]() { return ApplyTranslation(LOCTEXT("TransactionTopRight", "Top Right Translation"), FVector2D(1.0, 1.0)); })
 		]
 		+ SUniformGridPanel::Slot(0, 1)
 		.HAlign(HAlign_Center)
@@ -181,7 +185,7 @@ void FUVEditorUVTransformToolDetails::BuildQuickTranslateMenu(IDetailLayoutBuild
 			SNew(SButton)
 			.Text(LOCTEXT("QuickMoveLeft", "L"))
 			.ToolTipText(LOCTEXT("QuickMoveLeftToolTip", "Applies translation offset in the negative X axis."))
-			.OnClicked_Lambda([ApplyTranslation]() { return ApplyTranslation(FVector2D(-1.0, 0.0)); })
+			.OnClicked_Lambda([ApplyTranslation]() { return ApplyTranslation(LOCTEXT("TransactionLeft", "Left Translation"), FVector2D(-1.0, 0.0)); })
 		]
 		+ SUniformGridPanel::Slot(1, 1)
 		.HAlign(HAlign_Center)
@@ -197,7 +201,7 @@ void FUVEditorUVTransformToolDetails::BuildQuickTranslateMenu(IDetailLayoutBuild
 			SNew(SButton)
 			.Text(LOCTEXT("QuickMoveRight", "R"))
 			.ToolTipText(LOCTEXT("QuickMoveRightToolTip", "Applies the translation offset in the positive X axis."))
-			.OnClicked_Lambda([ApplyTranslation]() { return ApplyTranslation(FVector2D(1.0, 0.0)); })
+			.OnClicked_Lambda([ApplyTranslation]() { return ApplyTranslation(LOCTEXT("TransactionRight", "Right Translation"), FVector2D(1.0, 0.0)); })
 		]
 		+ SUniformGridPanel::Slot(0, 2)
 		.HAlign(HAlign_Center)
@@ -206,7 +210,7 @@ void FUVEditorUVTransformToolDetails::BuildQuickTranslateMenu(IDetailLayoutBuild
 			SNew(SButton)
 			.Text(LOCTEXT("QuickMoveBottomLeft", "BL"))
 			.ToolTipText(LOCTEXT("QuickMoveBottomLeftToolTip", "Applies translation offset in the negative X axis and the negative Y axis"))
-			.OnClicked_Lambda([ApplyTranslation]() { return ApplyTranslation(FVector2D(-1.0, -1.0)); })
+			.OnClicked_Lambda([ApplyTranslation]() { return ApplyTranslation(LOCTEXT("TransactionBottomLeft", "Bottom Left Translation"), FVector2D(-1.0, -1.0)); })
 		]
 		+ SUniformGridPanel::Slot(1, 2)
 		.HAlign(HAlign_Center)
@@ -215,7 +219,7 @@ void FUVEditorUVTransformToolDetails::BuildQuickTranslateMenu(IDetailLayoutBuild
 			SNew(SButton)
 			.Text(LOCTEXT("QuickMoveBottom", "B"))
 			.ToolTipText(LOCTEXT("QuickMoveBottomToolTip", "Applies the translation offset in the negative Y axis."))
-			.OnClicked_Lambda([ApplyTranslation]() { return ApplyTranslation(FVector2D(0.0, -1.0)); })
+			.OnClicked_Lambda([ApplyTranslation]() { return ApplyTranslation(LOCTEXT("TransactionBottom", "Bottom Translation"), FVector2D(0.0, -1.0)); })
 		]
 		+ SUniformGridPanel::Slot(2, 2)
 		.HAlign(HAlign_Center)
@@ -224,7 +228,7 @@ void FUVEditorUVTransformToolDetails::BuildQuickTranslateMenu(IDetailLayoutBuild
 			SNew(SButton)
 			.Text(LOCTEXT("QuickMoveBottomRight", "BR"))
 			.ToolTipText(LOCTEXT("QuickMoveBottomRightToolTip", "Applies the translation offset in the positive X axis and the negative Y axis."))
-			.OnClicked_Lambda([ApplyTranslation]() { return ApplyTranslation(FVector2D(1.0, -1.0)); })
+			.OnClicked_Lambda([ApplyTranslation]() { return ApplyTranslation(LOCTEXT("TransactionBottomRight", "Bottom Right Translation"), FVector2D(1.0, -1.0)); })
 		]
 
 	];
diff --git a/Engine/Plugins/Experimental/UVEditor/Source/UVEditorToolsEditorOnly/UVEditorToolsEditorOnly.Build.cs b/Engine/Plugins/Experimental/UVEditor/Source/UVEditorToolsEditorOnly/UVEditorToolsEditorOnly.Build.cs
--- a/Engine/Plugins/Experimental/UVEditor/Source/UVEditorToolsEditorOnly/UVEditorToolsEditorOnly.Build.cs
+++ b/Engine/Plugins/Experimental/UVEditor/Source/UVEditorToolsEditorOnly/UVEditorToolsEditorOnly.Build.cs
@@ -44,6 +44,7 @@ public class UVEditorToolsEditorOnly : ModuleRules
 				"Engine",
 				"Slate",
 				"SlateCore",
+				"UnrealEd",
 
 				"MeshModelingTools",
 				"MeshModelingToolsExp",
