UVEditor: Fix issue where selections for unwrapped meshes were being used for focusing the viewport to the selection in the 3d viewport. This change refactors unwrapped to applied selection conversion slightly to make it more convenient for areas of the code outside of the selection mechanic to perform conversions.

diff --git a/Engine/Plugins/Experimental/UVEditor/Source/UVEditor/Private/UVEditorMode.cpp b/Engine/Plugins/Experimental/UVEditor/Source/UVEditor/Private/UVEditorMode.cpp
--- a/Engine/Plugins/Experimental/UVEditor/Source/UVEditor/Private/UVEditorMode.cpp
+++ b/Engine/Plugins/Experimental/UVEditor/Source/UVEditor/Private/UVEditorMode.cpp
@@ -1310,7 +1310,7 @@ void UUVEditorMode::FocusLivePreviewCameraOnSelection()
 
 	for (const FUVToolSelection& Selection : CurrentSelections)
 	{
-		SelectionBoundingBox.Contain(Selection.ToBoundingBox(*Selection.Target->AppliedCanonical));
+		SelectionBoundingBox.Contain(Selection.GetConvertedSelectionForAppliedMesh().ToBoundingBox(*Selection.Target->AppliedCanonical));
 	}
 	for (const FUVToolSelection& Selection : CurrentUnsetSelections)
 	{
diff --git a/Engine/Plugins/Experimental/UVEditor/Source/UVEditorTools/Private/Selection/UVEditorMeshSelectionMechanic.cpp b/Engine/Plugins/Experimental/UVEditor/Source/UVEditorTools/Private/Selection/UVEditorMeshSelectionMechanic.cpp
--- a/Engine/Plugins/Experimental/UVEditor/Source/UVEditorTools/Private/Selection/UVEditorMeshSelectionMechanic.cpp
+++ b/Engine/Plugins/Experimental/UVEditor/Source/UVEditorTools/Private/Selection/UVEditorMeshSelectionMechanic.cpp
@@ -446,109 +446,55 @@ namespace UVEditorMeshSelectionMechanicLocals
 	TArray<int32> Convert3DHitsTo2DHits(ESelectionMode SelectionMode,
 		const FDynamicMesh3& UnwrapMesh, const FDynamicMesh3& AppliedMesh, const FDynamicMeshUVOverlay& UVOverlay, const TArray<int32>& IDsIn, TArray<int32>& AppliedMeshOnlyIDsOut)
 	{
-		TSet<int32> IDsOut;
-		AppliedMeshOnlyIDsOut.Empty();
-
-		auto FindUnwrapEdges = [&AppliedMesh, &UnwrapMesh](int32 Eid, TSet<int32>& OutUnwrapEdges)
-		{
-			bool bFoundEdges = false;
-			FIndex2i Triangles = AppliedMesh.GetEdgeT(Eid);
-			int32 Triangle0EdgeIndex = AppliedMesh.GetTriEdges(Triangles[0]).IndexOf(Eid);
-			int32 Triangle1EdgeIndex = Triangles[1] == IndexConstants::InvalidID ? IndexConstants::InvalidID : AppliedMesh.GetTriEdges(Triangles[1]).IndexOf(Eid);
-			if (UnwrapMesh.IsTriangle(Triangles[0]))
-			{
-				OutUnwrapEdges.Add(UnwrapMesh.GetTriEdge(Triangles[0], Triangle0EdgeIndex));
-				bFoundEdges = true;
-			}
-			if (UnwrapMesh.IsTriangle(Triangles[1]))
-			{
-				OutUnwrapEdges.Add(UnwrapMesh.GetTriEdge(Triangles[1], Triangle1EdgeIndex));
-				bFoundEdges = true;
-			}
-			return bFoundEdges;
-		};
-
+		TArray<int32> IDsOut;
 		switch (SelectionMode)
 		{
-		case ESelectionMode::Triangle:			
+		case ESelectionMode::Triangle:
 		case ESelectionMode::Island:
 		case ESelectionMode::Mesh:
-			for (int32 Tid : IDsIn)
-			{
-				if (UVOverlay.IsSetTriangle(Tid))
-				{
-					IDsOut.Add(Tid);
-				}
-				else
-				{
-					AppliedMeshOnlyIDsOut.Add(Tid);
-				}
-			}
+			IDsOut = FUVToolSelection::ConvertAppliedElementIdsToUnwrappedElementIds(FUVToolSelection::EType::Triangle, AppliedMesh, UnwrapMesh,
+				UVOverlay, IDsIn, AppliedMeshOnlyIDsOut);
 			break;
 		case ESelectionMode::Edge:
-			for (int32 Eid : IDsIn)
-			{				
-				if(!FindUnwrapEdges(Eid, IDsOut))				
-				{
-					AppliedMeshOnlyIDsOut.Add(Eid);
-				}
-			}
-			break;			
+			IDsOut = FUVToolSelection::ConvertAppliedElementIdsToUnwrappedElementIds(FUVToolSelection::EType::Edge, AppliedMesh, UnwrapMesh,
+		UVOverlay, IDsIn, AppliedMeshOnlyIDsOut);
+			break;
 		case ESelectionMode::Vertex:
-			for (int32 Vid : IDsIn)
-			{
-				TArray<int32> ElementsForVid;
-				UVOverlay.GetVertexElements(Vid, ElementsForVid);
-				IDsOut.Append(ElementsForVid);
-				if (ElementsForVid.IsEmpty())
-				{
-					AppliedMeshOnlyIDsOut.Add(Vid);
-				}
-			}
+			IDsOut = FUVToolSelection::ConvertAppliedElementIdsToUnwrappedElementIds(FUVToolSelection::EType::Vertex, AppliedMesh, UnwrapMesh,
+		UVOverlay, IDsIn, AppliedMeshOnlyIDsOut);
 			break;
 		default:
-			ensure(false);			
+			ensure(false);
 			break;
 		}
-		return IDsOut.Array();
+		return IDsOut;
 	}
 	
 	TArray<int32> Convert2DHitsTo3DHits(ESelectionMode SelectionMode,
 		const FDynamicMesh3& UnwrapMesh, const FDynamicMesh3& AppliedMesh, const FDynamicMeshUVOverlay& UVOverlay, TArray<int32>& IDsIn)
 	{
 		TArray<int32> IDsOut;
-		
-		auto FindAppliedEdge = [&UnwrapMesh, &AppliedMesh, &UVOverlay](int32 Eid)
-		{
-			int32 Vid1, Vid2;
-			FDynamicMesh3::FEdge EdgeInfo = UnwrapMesh.GetEdge(Eid);
-			Vid1 = UVOverlay.GetParentVertex(EdgeInfo.Vert[0]);
-			Vid2 = UVOverlay.GetParentVertex(EdgeInfo.Vert[1]);	
-			return AppliedMesh.FindEdge(Vid1, Vid2);
-		};
-
 		switch (SelectionMode)
 		{
 		case ESelectionMode::Triangle:
 		case ESelectionMode::Island:
 		case ESelectionMode::Mesh:
-			return IDsIn;
+			IDsOut = FUVToolSelection::ConvertUnwrappedElementIdsToAppliedElementIds(FUVToolSelection::EType::Triangle, UnwrapMesh, AppliedMesh,
+				UVOverlay, IDsIn);
+			break;
 		case ESelectionMode::Edge:
-			for (int32 Eid : IDsIn)
-			{
-				IDsOut.Add(FindAppliedEdge(Eid));
-			}
-			return IDsOut;
+			IDsOut = FUVToolSelection::ConvertUnwrappedElementIdsToAppliedElementIds(FUVToolSelection::EType::Edge, UnwrapMesh, AppliedMesh,
+				UVOverlay, IDsIn);
+			break;
 		case ESelectionMode::Vertex:
-			for (int32 Vid : IDsIn)
-			{
-				IDsOut.Add(UVOverlay.GetParentVertex(Vid));
-			}
-			return IDsOut;
+			IDsOut = FUVToolSelection::ConvertUnwrappedElementIdsToAppliedElementIds(FUVToolSelection::EType::Vertex, UnwrapMesh, AppliedMesh,
+				UVOverlay, IDsIn);
+			break;
 		default:
 			ensure(false);
-			return IDsOut;
+			break;
 		}
+		return IDsOut;
 	}
 
 
diff --git a/Engine/Plugins/Experimental/UVEditor/Source/UVEditorTools/Private/Selection/UVToolSelection.cpp b/Engine/Plugins/Experimental/UVEditor/Source/UVEditorTools/Private/Selection/UVToolSelection.cpp
--- a/Engine/Plugins/Experimental/UVEditor/Source/UVEditorTools/Private/Selection/UVToolSelection.cpp
+++ b/Engine/Plugins/Experimental/UVEditor/Source/UVEditorTools/Private/Selection/UVToolSelection.cpp
@@ -336,4 +336,154 @@ FUVToolSelection FUVToolSelection::GetConvertedSelection(const FDynamicMesh3& Me
 	return NewSelection;
 }
 
+FUVToolSelection FUVToolSelection::GetConvertedSelectionForUnwrappedMesh() const
+{
+	if (!Target.IsValid() || !Target->IsValid())
+	{
+		ensure(false);
+		return FUVToolSelection();
+	}
+
+	TArray<int32> UnsetAppliedIds;;
+	FUVToolSelection NewSelection;
+	NewSelection.Target = Target;
+	NewSelection.Type = Type;
+	NewSelection.SelectedIDs = TSet<int32>(ConvertAppliedElementIdsToUnwrappedElementIds(Type, *Target->AppliedCanonical, *Target->UnwrapCanonical,
+		*Target->AppliedCanonical->Attributes()->GetUVLayer(Target->UVLayerIndex), SelectedIDs.Array(), UnsetAppliedIds));
+
+	return NewSelection;
+}
+
+FUVToolSelection FUVToolSelection::GetConvertedSelectionForAppliedMesh() const
+{
+	if (!Target.IsValid() || !Target->IsValid())
+	{
+		ensure(false);
+		return FUVToolSelection();
+	}
+
+	FUVToolSelection NewSelection;
+	NewSelection.Target = Target;
+	NewSelection.Type = Type;
+	NewSelection.SelectedIDs = TSet<int32>(ConvertUnwrappedElementIdsToAppliedElementIds(Type, *Target->UnwrapCanonical, *Target->AppliedCanonical,
+		*Target->AppliedCanonical->Attributes()->GetUVLayer(Target->UVLayerIndex), SelectedIDs.Array()));
+
+	return NewSelection;
+}
+
+TArray<int32> FUVToolSelection::ConvertUnwrappedElementIdsToAppliedElementIds(EType SelectionMode,const FDynamicMesh3& UnwrapMesh, const FDynamicMesh3& AppliedMesh,
+	const FDynamicMeshUVOverlay& UVOverlay, const TArray<int32>& IDsIn)
+{
+	TArray<int32> IDsOut;
+
+	auto FindAppliedEdge = [&UnwrapMesh, &AppliedMesh, &UVOverlay](int32 Eid)
+	{
+		int32 Vid1, Vid2;
+		FDynamicMesh3::FEdge EdgeInfo = UnwrapMesh.GetEdge(Eid);
+		Vid1 = UVOverlay.GetParentVertex(EdgeInfo.Vert[0]);
+		Vid2 = UVOverlay.GetParentVertex(EdgeInfo.Vert[1]);
+		if (Vid1 != IndexConstants::InvalidID && Vid2 != IndexConstants::InvalidID)
+		{
+			return AppliedMesh.FindEdge(Vid1, Vid2);
+		}
+		else
+		{
+			return IndexConstants::InvalidID;
+		}
+	};
+
+	switch (SelectionMode)
+	{
+	case EType::Triangle:
+		return IDsIn;
+	case EType::Edge:
+		for (int32 Eid : IDsIn)
+		{
+			int32 AppliedEid = FindAppliedEdge(Eid);
+			if (AppliedEid != IndexConstants::InvalidID)
+			{
+				IDsOut.Add(AppliedEid);
+			}
+		}
+		return IDsOut;
+	case EType::Vertex:
+		for (int32 Vid : IDsIn)
+		{
+			IDsOut.Add(UVOverlay.GetParentVertex(Vid));
+		}
+		return IDsOut;
+	default:
+		ensure(false);
+		return IDsOut;
+	}
+}
+
+TArray<int32> FUVToolSelection::ConvertAppliedElementIdsToUnwrappedElementIds(EType SelectionMode,	const FDynamicMesh3& AppliedMesh, const FDynamicMesh3& UnwrapMesh,
+	const FDynamicMeshUVOverlay& UVOverlay, const TArray<int32>& IDsIn, TArray<int32>& AppliedMeshOnlyIDsOut)
+{
+	TSet<int32> IDsOut;
+	AppliedMeshOnlyIDsOut.Empty();
+
+	auto FindUnwrapEdges = [&AppliedMesh, &UnwrapMesh](int32 Eid, TSet<int32>& OutUnwrapEdges)
+	{
+		bool bFoundEdges = false;
+		FIndex2i Triangles = AppliedMesh.GetEdgeT(Eid);
+		int32 Triangle0EdgeIndex = AppliedMesh.GetTriEdges(Triangles[0]).IndexOf(Eid);
+		int32 Triangle1EdgeIndex = Triangles[1] == IndexConstants::InvalidID ? IndexConstants::InvalidID : AppliedMesh.GetTriEdges(Triangles[1]).IndexOf(Eid);
+		if (UnwrapMesh.IsTriangle(Triangles[0]))
+		{
+			OutUnwrapEdges.Add(UnwrapMesh.GetTriEdge(Triangles[0], Triangle0EdgeIndex));
+			bFoundEdges = true;
+		}
+		if (UnwrapMesh.IsTriangle(Triangles[1]))
+		{
+			OutUnwrapEdges.Add(UnwrapMesh.GetTriEdge(Triangles[1], Triangle1EdgeIndex));
+			bFoundEdges = true;
+		}
+		return bFoundEdges;
+	};
+
+	switch (SelectionMode)
+	{
+	case EType::Triangle:
+		for (int32 Tid : IDsIn)
+		{
+			if (UVOverlay.IsSetTriangle(Tid))
+			{
+				IDsOut.Add(Tid);
+			}
+			else
+			{
+				AppliedMeshOnlyIDsOut.Add(Tid);
+			}
+		}
+		break;
+	case EType::Edge:
+		for (int32 Eid : IDsIn)
+		{
+			if (!FindUnwrapEdges(Eid, IDsOut))
+			{
+				AppliedMeshOnlyIDsOut.Add(Eid);
+			}
+		}
+		break;
+	case EType::Vertex:
+		for (int32 Vid : IDsIn)
+		{
+			TArray<int32> ElementsForVid;
+			UVOverlay.GetVertexElements(Vid, ElementsForVid);
+			IDsOut.Append(ElementsForVid);
+			if (ElementsForVid.IsEmpty())
+			{
+				AppliedMeshOnlyIDsOut.Add(Vid);
+			}
+		}
+		break;
+	default:
+		ensure(false);
+		break;
+	}
+	return IDsOut.Array();
+}
+
 #undef LOCTEXT_NAMESPACE
\ No newline at end of file
diff --git a/Engine/Plugins/Experimental/UVEditor/Source/UVEditorTools/Public/Selection/UVToolSelection.h b/Engine/Plugins/Experimental/UVEditor/Source/UVEditorTools/Public/Selection/UVToolSelection.h
--- a/Engine/Plugins/Experimental/UVEditor/Source/UVEditorTools/Public/Selection/UVToolSelection.h
+++ b/Engine/Plugins/Experimental/UVEditor/Source/UVEditorTools/Public/Selection/UVToolSelection.h
@@ -71,11 +71,51 @@ public:
 	bool AreElementsPresentInMesh(const FDynamicMesh3& Mesh) const;
 
 	FAxisAlignedBox3d ToBoundingBox(const FDynamicMesh3& Mesh) const;
+
 	/**
 	 * Utility method to construct a new selection from an existing reference selection, only with a new geometry type.
 	 */
 	FUVToolSelection GetConvertedSelection(const FDynamicMesh3& Mesh, FUVToolSelection::EType ExpectedSelectionType) const;
 
+	/** 
+	* Utility method to construct a new selection from an existing reference Applied mesh selection, only for the equivalent Unwrap mesh geometry.
+	* 
+	* If the target for this selection is currently invalid or unset, this will return an empty selection.
+	*/
+	FUVToolSelection GetConvertedSelectionForUnwrappedMesh() const;
+
+	/**
+	* Utility method to construct a new selection from an existing reference Unwrap mesh selection, only for the equivalent Applied mesh geometry.
+	*
+	* If the target for this selection is currently invalid or unset, this will return an empty selection.
+	*/
+	FUVToolSelection GetConvertedSelectionForAppliedMesh() const;
+
+	/**
+	 * Utility method to convert element ids from those belonging to the unwrapped mesh to those belonging to equivalent ids on the applied mesh
+	 * 
+	 * @param SelectionMode The type of element the input ids represent
+	 * @param UnwrapMesh A dynamic mesh instance representing the unwrapped mesh which the input ids belong to 
+	 * @param AppliedMesh A dynamic mesh instance representing the applied mesh for which the resulting ids will belong to
+	 * @param UVOverlay An overlay representing the current UV mapping linking the unwrapped and applied mesh
+	 * @param IDsIn The ids of the elements from the unwrapped mesh to convert to applied element ids
+	 */
+	static TArray<int32> ConvertUnwrappedElementIdsToAppliedElementIds(EType SelectionMode, const FDynamicMesh3& UnwrapMesh,
+		const FDynamicMesh3& AppliedMesh, const FDynamicMeshUVOverlay& UVOverlay, const TArray<int32>& IDsIn);
+
+	/**
+	 * Utility method to convert element ids from those belonging to the applied mesh to those belonging to equivalent ids on the unwrapped mesh
+	 *
+	 * @param SelectionMode The type of element the input ids represent
+ 	 * @param AppliedMesh A dynamic mesh instance representing the applied mesh which the input ids belong to
+	 * @param UnwrapMesh A dynamic mesh instance representing the unwrapped mesh for which the resulting ids will belong to
+	 * @param UVOverlay An overlay representing the current UV mapping linking the unwrapped and applied mesh
+	 * @param IDsIn The ids of the elements from the unwrapped mesh to convert to applied element ids
+	 * @param AppliedOnlyElementIds The ids of any remaining elements that have no direct mappings, due to unset UVs, in the unwrapped mesh
+	 */
+	static TArray<int32> ConvertAppliedElementIdsToUnwrappedElementIds(EType SelectionMode, const FDynamicMesh3& AppliedMesh,
+		const FDynamicMesh3& UnwrapMesh, const FDynamicMeshUVOverlay& UVOverlay, const TArray<int32>& IDsIn, TArray<int32>& AppliedOnlyElementIds);
+
 protected:
 	
 	FMeshEdgesFromTriangleSubIndices StableEdgeIDs;
