UVEditor: Clean up tooltip text in the Transform tool to be more consistent between different buttons.

diff --git a/Engine/Plugins/Experimental/UVEditor/Source/UVEditorToolsEditorOnly/Private/DetailsCustomizations/UVTransformToolCustomizations.cpp b/Engine/Plugins/Experimental/UVEditor/Source/UVEditorToolsEditorOnly/Private/DetailsCustomizations/UVTransformToolCustomizations.cpp
--- a/Engine/Plugins/Experimental/UVEditor/Source/UVEditorToolsEditorOnly/Private/DetailsCustomizations/UVTransformToolCustomizations.cpp
+++ b/Engine/Plugins/Experimental/UVEditor/Source/UVEditorToolsEditorOnly/Private/DetailsCustomizations/UVTransformToolCustomizations.cpp
@@ -184,7 +184,7 @@ void FUVEditorUVTransformToolDetails::BuildQuickTranslateMenu(IDetailLayoutBuild
 		[
 			SNew(SButton)
 			.Text(LOCTEXT("QuickMoveLeft", "L"))
-			.ToolTipText(LOCTEXT("QuickMoveLeftToolTip", "Applies translation offset in the negative X axis."))
+			.ToolTipText(LOCTEXT("QuickMoveLeftToolTip", "Applies the translation offset in the negative X axis."))
 			.OnClicked_Lambda([ApplyTranslation]() { return ApplyTranslation(LOCTEXT("TransactionLeft", "Left Translation"), FVector2D(-1.0, 0.0)); })
 		]
 		+ SUniformGridPanel::Slot(1, 1)
@@ -209,7 +209,7 @@ void FUVEditorUVTransformToolDetails::BuildQuickTranslateMenu(IDetailLayoutBuild
 		[
 			SNew(SButton)
 			.Text(LOCTEXT("QuickMoveBottomLeft", "BL"))
-			.ToolTipText(LOCTEXT("QuickMoveBottomLeftToolTip", "Applies translation offset in the negative X axis and the negative Y axis"))
+			.ToolTipText(LOCTEXT("QuickMoveBottomLeftToolTip", "Applies the translation offset in the negative X axis and the negative Y axis"))
 			.OnClicked_Lambda([ApplyTranslation]() { return ApplyTranslation(LOCTEXT("TransactionBottomLeft", "Bottom Left Translation"), FVector2D(-1.0, -1.0)); })
 		]
 		+ SUniformGridPanel::Slot(1, 2)
