UVEditor: Fix a tool tip for the Transform tool that had gone missed in the last cleanup pass.

diff --git a/Engine/Plugins/Experimental/UVEditor/Source/UVEditorToolsEditorOnly/Private/DetailsCustomizations/UVTransformToolCustomizations.cpp b/Engine/Plugins/Experimental/UVEditor/Source/UVEditorToolsEditorOnly/Private/DetailsCustomizations/UVTransformToolCustomizations.cpp
--- a/Engine/Plugins/Experimental/UVEditor/Source/UVEditorToolsEditorOnly/Private/DetailsCustomizations/UVTransformToolCustomizations.cpp
+++ b/Engine/Plugins/Experimental/UVEditor/Source/UVEditorToolsEditorOnly/Private/DetailsCustomizations/UVTransformToolCustomizations.cpp
@@ -157,7 +157,7 @@ void FUVEditorUVTransformToolDetails::BuildQuickTranslateMenu(IDetailLayoutBuild
 		[
 			SNew(SButton)
 			.Text(LOCTEXT("QuickMoveTopLeft", "TL"))
-			.ToolTipText(LOCTEXT("QuickMoveTopLeftToolTip", "Applies translation offset in the negative X axis and the positive Y axis"))
+			.ToolTipText(LOCTEXT("QuickMoveTopLeftToolTip", "Applies the translation offset in the negative X axis and the positive Y axis"))
 			.OnClicked_Lambda([ApplyTranslation]() { return ApplyTranslation(LOCTEXT("TransactionTopLeft", "Top Left Translation"), FVector2D(-1.0, 1.0)); })
 		]
 		+ SUniformGridPanel::Slot(1, 0)
