UV Editor: Fix issue where the background grid was not fully being drawn in the Y-Axis direction due to incorrect handling of Y minimum and maximum bounds.

diff --git a/Engine/Plugins/Experimental/UVEditor/Source/UVEditor/Private/UVEditor2DViewportClient.cpp b/Engine/Plugins/Experimental/UVEditor/Source/UVEditor/Private/UVEditor2DViewportClient.cpp
--- a/Engine/Plugins/Experimental/UVEditor/Source/UVEditor/Private/UVEditor2DViewportClient.cpp
+++ b/Engine/Plugins/Experimental/UVEditor/Source/UVEditor/Private/UVEditor2DViewportClient.cpp
@@ -252,6 +252,10 @@ void FUVEditor2DViewportClient::Draw(const FSceneView* View, FPrimitiveDrawInter
 
 void FUVEditor2DViewportClient::DrawGrid(const FSceneView* View, FPrimitiveDrawInterface* PDI)
 {
+	// TODO: Drawing the grid is another location where we need to pay attention to the coordinate orienation of the viewport and camera. This
+	// should probably be cleaned up at some point to make it more robust to these assumptions. 
+
+
 	// Basic scaling amount
 	const double UVScale = UUVEditorMode::GetUVMeshScalingFactor();
 	
@@ -267,6 +271,11 @@ void FUVEditor2DViewportClient::DrawGrid(const FSceneView* View, FPrimitiveDrawI
 	ViewLoc.Z = 0.0; // We are treating the scene like a 2D plane, so we'll clamp the Z position here to 
 	               // 0 as a simple projection step just in case.
 
+	// We wrap these inside an FAxisAlignedBox3d to ensure the Min coordinates < Max coordinates, which might not be the case coming directly out of the View.
+	UE::Geometry::FAxisAlignedBox3d  WorldBounds;
+	WorldBounds.Contain(WorldBoundsMax);
+	WorldBounds.Contain(WorldBoundsMin);
+	
 	// Prevent grid from drawing if we are too close or too far, in order to avoid potential graphical issues.
 	if (ZoomFactor < 100000 && ZoomFactor > 1)
 	{
@@ -278,7 +287,7 @@ void FUVEditor2DViewportClient::DrawGrid(const FSceneView* View, FPrimitiveDrawI
 		Colors.Push(FUVEditorUXSettings::GridMinorColor);
 		MeshDebugDraw::DrawHierarchicalGrid(UVScale, ZoomFactor / UVScale,
 			500, // Maximum density of lines to draw per level before skipping the level
-			WorldBoundsMax, WorldBoundsMin,
+			WorldBounds.Max, WorldBounds.Min,
 			FUVEditorUXSettings::GridLevels, // Number of levels to draw
 			FUVEditorUXSettings::GridSubdivisionsPerLevel, // Number of subdivisions per level
 			Colors,
