UVEditor: Fix ensures being fired when editor is closed without live preview hover being terminated.

diff --git a/Engine/Plugins/Experimental/UVEditor/Source/UVEditorTools/Private/Selection/UVEditorMeshSelectionMechanic.cpp b/Engine/Plugins/Experimental/UVEditor/Source/UVEditorTools/Private/Selection/UVEditorMeshSelectionMechanic.cpp
--- a/Engine/Plugins/Experimental/UVEditor/Source/UVEditorTools/Private/Selection/UVEditorMeshSelectionMechanic.cpp
+++ b/Engine/Plugins/Experimental/UVEditor/Source/UVEditorTools/Private/Selection/UVEditorMeshSelectionMechanic.cpp
@@ -750,6 +750,14 @@ void UUVEditorMeshSelectionMechanic::SetShowHoveredElements(bool bShow)
 
 void UUVEditorMeshSelectionMechanic::Shutdown()
 {
+	// Note that the ForceTerminateSource should happen before destroying everything else so that we
+	// can shut down behaviors properly.
+	if (LivePreviewInputRouter.IsValid())
+	{
+		LivePreviewInputRouter->ForceTerminateSource(LivePreviewBehaviorSource);
+		LivePreviewInputRouter->DeregisterSource(LivePreviewBehaviorSource);
+	}
+
 	if (HoverGeometryActor)
 	{
 		HoverGeometryActor->Destroy();
@@ -760,10 +768,6 @@ void UUVEditorMeshSelectionMechanic::Shutdown()
 		LivePreviewHoverGeometryActor->Destroy();
 		LivePreviewHoverGeometryActor = nullptr;
 	}
-	if (LivePreviewAPI)
-	{
-		LivePreviewInputRouter->DeregisterSource(LivePreviewBehaviorSource);
-	}
 
 	SelectionAPI = nullptr;
 	ViewportButtonsAPI = nullptr;
