UVEditor: remove hack that causes user to see broken values in numerical display when scaling. This makes the scaling handles more sensitive (same sensitivity as with snapping enabled), but is better than being broken.

diff --git a/Engine/Plugins/Experimental/UVEditor/Source/UVEditorTools/Private/UVSelectTool.cpp b/Engine/Plugins/Experimental/UVEditor/Source/UVEditorTools/Private/UVSelectTool.cpp
--- a/Engine/Plugins/Experimental/UVEditor/Source/UVEditorTools/Private/UVSelectTool.cpp
+++ b/Engine/Plugins/Experimental/UVEditor/Source/UVEditorTools/Private/UVSelectTool.cpp
@@ -627,19 +627,6 @@ void UUVSelectTool::ApplyGizmoTransform()
 {
 	FTransformSRT3d TransformToApply(UnappliedGizmoTransform);
 
-	// TODO: The division here is a bit of a hack. Properly-speaking, the scaling handles should act relative to
-	// gizmo size, not the visible space across which we drag, otherwise it becomes dependent on the units we
-	// use and our absolute distance from the object. Since our UV unwrap is scaled by 1000 to make it
-	// easier to zoom in and out without running into issues, the measure of the distance across which we typically
-	// drag the handles is too high to be convenient. Until we make the scaling invariant to units/distance from
-	// target, we use this hack.
-	// Note: If we're using the snapping feature we don't want to mess with the values on scale here, since the
-	// snapping has already been performed. If we use the hack, we won't have the exact snapped scale the user expects.
-	if (!ViewportButtonsAPI->GetSnapEnabled(UUVToolViewportButtonsAPI::ESnapTypeFlag::Scale))
-	{
-		TransformToApply.SetScale(FVector::One() + (UnappliedGizmoTransform.GetScale3D() - FVector::One()) / 10);
-	}
-
 	for (int32 SelectionIndex = 0; SelectionIndex < CurrentSelections.Num(); ++SelectionIndex)
 	{
 		UUVEditorToolMeshInput* Target = CurrentSelections[SelectionIndex].Target.Get();
