UVEditor: Fix gizmo being affected by level editor rotation snap settings instead of the uv editor settings.

diff --git a/Engine/Plugins/Experimental/UVEditor/Source/UVEditorTools/Private/UVSelectTool.cpp b/Engine/Plugins/Experimental/UVEditor/Source/UVEditorTools/Private/UVSelectTool.cpp
--- a/Engine/Plugins/Experimental/UVEditor/Source/UVEditorTools/Private/UVSelectTool.cpp
+++ b/Engine/Plugins/Experimental/UVEditor/Source/UVEditorTools/Private/UVSelectTool.cpp
@@ -245,8 +245,15 @@ void UUVSelectTool::Setup()
 	TransformProxy->OnTransformChanged.AddUObject(this, &UUVSelectTool::GizmoTransformChanged);
 	TransformProxy->OnEndTransformEdit.AddUObject(this, &UUVSelectTool::GizmoTransformEnded);
 
-	// Always align gizmo to x and y axes
+	// TODO: Perhaps ideally we would be able to use these built-in options for gizmo, but they would require us to
+	// use our own IToolsContextQueriesAPI implementation that does not query GEditor for the snapping settings, since
+	// the UV editor needs to have its own snap settings. That might be the proper thing to do in the long term. 
+	// Currently, however, the UV editor manages and queries the snap settings "by hand" through the viewport buttons
+	// api, and letting the gizmo query the editor will give us the wrong value.
+	TransformGizmo->bSnapToWorldGrid = false;
+	TransformGizmo->bSnapToWorldRotGrid = false;
 	TransformGizmo->bUseContextCoordinateSystem = false;
+
 	TransformGizmo->SetActiveTarget(TransformProxy, GetToolManager());
 	
 	// Tell the gizmo to be drawn on top even over translucent-mode materials.
