UVEditor: Enable support for a new UV distortion visualization display setting.

When enabled, it computes a per triangle distortion value of the triangle's UVs vs it's world space positions. This is then displayed as a color blind friendly color overlay on the Unwrap mesh. The overlay is designed to update automatically whenever the preview is updated, supporting both tool operation edits and manual updates with the gizmo.

diff --git a/Engine/Plugins/Experimental/UVEditor/Source/UVEditor/Private/UVEditorDistortionVisualization.cpp b/Engine/Plugins/Experimental/UVEditor/Source/UVEditor/Private/UVEditorDistortionVisualization.cpp
new file mode 100644
--- /dev/null
+++ b/Engine/Plugins/Experimental/UVEditor/Source/UVEditor/Private/UVEditorDistortionVisualization.cpp
@@ -0,0 +1,250 @@
+// Copyright Epic Games, Inc. All Rights Reserved.
+
+#include "UVEditorDistortionVisualization.h"
+#include "DynamicMesh/DynamicMesh3.h"
+#include "DynamicMesh/DynamicMeshAttributeSet.h"
+#include "Drawing/MeshWireframeComponent.h"
+#include "Engine/Texture2D.h"
+#include "Materials/Material.h"
+#include "ToolSetupUtil.h"
+#include "Async/Async.h"
+#include "UVEditorUXSettings.h"
+#include "UDIMUtilities.h"
+#include "Materials/MaterialExpressionTextureBase.h"
+#include "ToolTargets/UVEditorToolMeshInput.h"
+#include "MeshOpPreviewHelpers.h"
+#include "UVEditorUXSettings.h"
+#include "Math/UnitConversion.h"
+#include "Parameterization/MeshUDIMClassifier.h"
+#include "UDIMUtilities.h"
+#include "Math/UVMetrics.h"
+
+#include UE_INLINE_GENERATED_CPP_BY_NAME(UVEditorDistortionVisualization)
+
+using namespace UE::Geometry;
+
+namespace DistortionMetrics
+{
+	double Triangle2DArea(FVector2f p1, FVector2f p2, FVector2f p3)
+	{
+		return FMath::Abs(((p2.X - p1.X) * (p3.Y - p1.Y) - (p3.X - p1.X) * (p2.Y - p1.Y)) / 2.0);
+	}
+
+	double Triangle3DArea(FVector3d q1, FVector3d q2, FVector3d q3)
+	{
+		FVector3d Q1 = FVector3d(q2.X - q1.X, q2.Y - q1.Y, q2.Z - q1.Z);
+		FVector3d Q2 = FVector3d(q3.X - q1.X, q3.Y - q1.Y, q3.Z - q1.Z);
+
+		FVector3d U = Q1.Cross(Q2);
+
+		return U.Length() / 2.0;
+	}
+
+	double ReedBeta(const FDynamicMesh3& Mesh, int UVChannel, int Tid)
+	{
+		return 1 - FUVMetrics::ReedBeta(Mesh, UVChannel, Tid); // Flip the value, so green is 0, or non-distorted
+	}
+
+	double Sander(const FDynamicMesh3& Mesh, int UVChannel, int Tid, bool bUseL2)
+	{	
+			return 2.0 - FMath::Clamp(FUVMetrics::Sander(Mesh, UVChannel, Tid, bUseL2), 1.0, 2.0);
+	}
+
+	double TexelDensity(const FDynamicMesh3& Mesh, int UVChannel, int Tid, int32 MapSize, float TargetTexelDensity)
+	{
+		double Density = FUVMetrics::TexelDensity(Mesh, UVChannel, Tid, MapSize);
+
+		// We're looking to map the density into a divergent map here, with values ranging from -0.5 to 0.5... Additionally we'd like the scaling factor to be equal in both direction.
+		// So we're going to ensure that a 50% growth or reduction in size maps the same on the scale. This makes the math slightly more complex than just a relative difference.
+		// This represents a mapping of [0.5*D, 1.5*D] where D is target density linearly to [-0.5,0.5]
+
+		double DensityDifference = FMath::Abs(TargetTexelDensity - Density);
+		double DensityRatio = DensityDifference < (TargetTexelDensity / 2.0f) ? DensityDifference / TargetTexelDensity : 0.5;
+		return DensityRatio * FMath::Sign(TargetTexelDensity - Density);
+	}
+}
+
+
+
+void UUVEditorDistortionVisualization::Initialize()
+{
+	Settings = NewObject<UUVEditorDistortionVisualizationProperties>(this);
+	Settings->WatchProperty(Settings->bVisible, [this](bool) { bSettingsModified = true; });
+	Settings->WatchProperty(Settings->Metric, [this](EDistortionMetric) { bSettingsModified = true; });
+	Settings->WatchProperty(Settings->MapSize, [this](int32) { bSettingsModified = true; });
+	Settings->WatchProperty(Settings->TargetTexelDensity, [this](float) { bSettingsModified = true; });
+	Settings->WatchProperty(Settings->bCompareToAverageDensity, [this](bool) { bSettingsModified = true; });
+	Settings->WatchProperty(Settings->bRespectUDIMTextureResolutions, [this](bool) { bSettingsModified = true; });
+	Settings->WatchProperty(Settings->PerUDIMTextureResolution,
+		                    [this](const TMap<int32, int32>&) { bSettingsModified = true; }, 
+		                    [this](const TMap<int32, int32>& A, const TMap<int32, int32>& B) -> bool { return !A.OrderIndependentCompareEqual(B); });
+	bSettingsModified = false;
+
+	PerTargetAverageTexelDensity.SetNum(Targets.Num());
+}
+
+void UUVEditorDistortionVisualization::ComputeInitialMeshSurfaceAreas()
+{
+	double WorldSpaceSurfaceArea = 0.0;
+	double UVSpaceSurfaceArea = 0.0;
+	FVector3d q1, q2, q3;
+	FVector2f p1, p2, p3;
+
+	MaximumTileTextureResolution = 1000;
+	if (Settings->bRespectUDIMTextureResolutions && !Settings->PerUDIMTextureResolution.IsEmpty())
+	{
+		MaximumTileTextureResolution = 0;
+		for (TPair<int32, int32>& TileResolution : Settings->PerUDIMTextureResolution)
+		{
+			MaximumTileTextureResolution = FMath::Max(MaximumTileTextureResolution, TileResolution.Value);
+		}
+	}
+
+	for (int32 TargetIndex = 0; TargetIndex < Targets.Num(); ++TargetIndex)
+	{
+		TMap<int32, double> WorldSpaceSurfaceAreaPerTileResolution;
+		TMap<int32, double> UVSpaceSurfaceAreaPerTileResolution;
+		TMap<int32, double> AverageDensityPerTileResolution;
+		TMap<int32, int32> CountPerTileResolution;
+
+		TObjectPtr<UUVEditorToolMeshInput>& Target = Targets[TargetIndex];
+		for (int32 Tid = 0; Tid < Target->AppliedCanonical->TriangleCount(); ++Tid)
+		{
+			const int32* TileTextureResolution = nullptr;
+			if(Settings->bRespectUDIMTextureResolutions)
+			{
+				TArray<int32> TidsToClassify;
+				TidsToClassify.Add(Tid);
+				FVector2i UDIMCoords = FDynamicMeshUDIMClassifier::ClassifyTrianglesToUDIM(Target->AppliedCanonical->Attributes()->GetUVLayer(Target->UVLayerIndex), TidsToClassify);
+				int32 UDIMBlockIndex = UE::TextureUtilitiesCommon::GetUDIMIndex(UDIMCoords.X, UDIMCoords.Y);
+				TileTextureResolution = Settings->PerUDIMTextureResolution.Find(UDIMBlockIndex);
+			}
+			if (!TileTextureResolution)
+			{
+				TileTextureResolution = &MaximumTileTextureResolution;
+			}
+
+			
+			Target->AppliedCanonical->GetTriVertices(Tid, q1, q2, q3);
+			Target->AppliedCanonical->Attributes()->GetUVLayer(Target->UVLayerIndex)->GetTriElements(Tid, p1, p2, p3);
+
+			WorldSpaceSurfaceAreaPerTileResolution.FindOrAdd(*TileTextureResolution) += DistortionMetrics::Triangle3DArea(q1,q2,q3);
+			UVSpaceSurfaceAreaPerTileResolution.FindOrAdd(*TileTextureResolution) += DistortionMetrics::Triangle2DArea(p1, p2, p3);
+			CountPerTileResolution.FindOrAdd(*TileTextureResolution)++;			
+
+			AverageDensityPerTileResolution.FindOrAdd(*TileTextureResolution) += (DistortionMetrics::Triangle2DArea(p1, p2, p3) / DistortionMetrics::Triangle3DArea(q1, q2, q3));			
+		}
+		TArray<int32> Resolutions;
+		WorldSpaceSurfaceAreaPerTileResolution.GenerateKeyArray(Resolutions);
+		PerTargetAverageTexelDensity[TargetIndex] = 0.0;
+		for (int32 Resolution : Resolutions)
+		{
+			PerTargetAverageTexelDensity[TargetIndex] += Resolution * FMath::Sqrt(AverageDensityPerTileResolution[Resolution] / CountPerTileResolution[Resolution]);
+		}
+		PerTargetAverageTexelDensity[TargetIndex] /= Resolutions.Num();
+	}
+}
+
+void UUVEditorDistortionVisualization::Shutdown()
+{
+}
+
+void UUVEditorDistortionVisualization::OnTick(float DeltaTime)
+{
+
+	if (bSettingsModified)
+	{		
+		UpdateVisibility();
+		bSettingsModified = false;
+	}
+}
+
+void UUVEditorDistortionVisualization::UpdateVisibility()
+{
+	if (Settings->bVisible && Settings->Metric == EDistortionMetric::TexelDensity)
+	{
+		ComputeInitialMeshSurfaceAreas();
+	}
+
+	for (int32 TargetIndex = 0; TargetIndex < Targets.Num(); ++TargetIndex)
+	{
+		TObjectPtr<UUVEditorToolMeshInput>& Target = Targets[TargetIndex];
+
+		if (Settings->bVisible)
+		{
+			ConfigureMeshColorsForTarget(TargetIndex);
+		}
+		else
+		{
+			Target->UnwrapPreview->PreviewMesh->ClearTriangleColorFunction(UPreviewMesh::ERenderUpdateMode::FastUpdate);
+		}
+	}
+}
+
+void UUVEditorDistortionVisualization::ConfigureMeshColorsForTarget(int32 TargetIndex)
+{
+	TObjectPtr<UUVEditorToolMeshInput>& Target = Targets[TargetIndex];
+
+	Target->UnwrapPreview->PreviewMesh->SetTriangleColorFunction([this, &Target, TargetIndex](const FDynamicMesh3* Mesh, int Tid)
+	{
+		FColor TidDistortionColor;
+		Target->AppliedPreview->PreviewMesh->ProcessMesh([this, &TidDistortionColor, Tid, &Target, TargetIndex](const FDynamicMesh3& AppliedMesh)
+			{
+				TidDistortionColor = GetDistortionColorForTriangle(AppliedMesh, Target->UVLayerIndex, Tid, PerTargetAverageTexelDensity[TargetIndex]);
+			});
+		return TidDistortionColor;
+	}, UPreviewMesh::ERenderUpdateMode::FastUpdate);
+}
+
+FColor UUVEditorDistortionVisualization::GetDistortionColorForTriangle(const FDynamicMesh3& Mesh, int UVChannel, int Tid, double TargetAverageTexelDensity)
+{
+	double MetricValue = 0;
+	FColor BackgroundColor = FColor::Black;
+
+	switch (Settings->Metric)
+	{
+	case EDistortionMetric::ReedBeta:
+		MetricValue = DistortionMetrics::ReedBeta(Mesh, UVChannel, Tid);
+		BackgroundColor = FUVEditorUXSettings::MakeCividisColorFromScalar(MetricValue);
+		break;
+	case EDistortionMetric::Sander_L2:
+		MetricValue = DistortionMetrics::Sander(Mesh, UVChannel, Tid, true);
+		BackgroundColor = FUVEditorUXSettings::MakeCividisColorFromScalar(MetricValue);
+		break;
+	case EDistortionMetric::Sander_LInf:
+		MetricValue = DistortionMetrics::Sander(Mesh, UVChannel, Tid, false);
+		BackgroundColor = FUVEditorUXSettings::MakeCividisColorFromScalar(MetricValue);
+		break;
+	case EDistortionMetric::TexelDensity:
+	{
+		int32* TileTextureResolution = nullptr;
+		if (Settings->bRespectUDIMTextureResolutions)
+		{
+			TArray<int32> TidsToClassify;
+			TidsToClassify.Add(Tid);
+			FVector2i UDIMCoords = FDynamicMeshUDIMClassifier::ClassifyTrianglesToUDIM(Mesh.Attributes()->GetUVLayer(UVChannel), TidsToClassify);
+			int32 UDIMBlockIndex = UE::TextureUtilitiesCommon::GetUDIMIndex(UDIMCoords.X, UDIMCoords.Y);
+			TileTextureResolution = Settings->PerUDIMTextureResolution.Find(UDIMBlockIndex);
+		}
+
+		if (Settings->bCompareToAverageDensity)
+		{
+			MetricValue = DistortionMetrics::TexelDensity(Mesh, UVChannel, Tid, TileTextureResolution ? *TileTextureResolution : MaximumTileTextureResolution /* Use some artifical "pixel density" here to avoid small numbers */, TargetAverageTexelDensity);
+		}
+		else
+		{
+			MetricValue = DistortionMetrics::TexelDensity(Mesh, UVChannel, Tid, TileTextureResolution ? *TileTextureResolution : Settings->MapSize, Settings->TargetTexelDensity);
+		}
+		BackgroundColor = FUVEditorUXSettings::MakeTurboColorFromScalar(MetricValue);
+	}
+		break;
+	default:
+		ensure(false);
+	}
+
+	
+	return BackgroundColor;
+
+}
+
+
diff --git a/Engine/Plugins/Experimental/UVEditor/Source/UVEditor/Private/UVEditorMode.cpp b/Engine/Plugins/Experimental/UVEditor/Source/UVEditor/Private/UVEditorMode.cpp
--- a/Engine/Plugins/Experimental/UVEditor/Source/UVEditor/Private/UVEditorMode.cpp
+++ b/Engine/Plugins/Experimental/UVEditor/Source/UVEditor/Private/UVEditorMode.cpp
@@ -44,6 +44,7 @@
 #include "UVEditorSubsystem.h"
 #include "ContextObjects/UVToolContextObjects.h"
 #include "UVEditorBackgroundPreview.h"
+#include "UVEditorDistortionVisualization.h"
 #include "UVEditorToolAnalyticsUtils.h"
 #include "UVEditorUXSettings.h"
 #include "UDIMUtilities.h"
@@ -271,7 +272,7 @@ void UUVEditorUDIMProperties::UpdateActiveUDIMsFromTexture()
 			UDIMSourceTexture->Source.GetBlock(Block, SourceBlock);
 
 			ActiveUDIMs.Add(FUDIMSpecifier({ UE::TextureUtilitiesCommon::GetUDIMIndex(SourceBlock.BlockX, SourceBlock.BlockY),
-				                             SourceBlock.BlockX, SourceBlock.BlockY }));
+				                             SourceBlock.BlockX, SourceBlock.BlockY, FMath::Min(SourceBlock.SizeX, SourceBlock.SizeY) }));
 		}
 	}
 	CheckAndUpdateWatched();
@@ -393,7 +394,7 @@ void UUVEditorMode::Enter()
 
 	BackgroundVisualization->Settings->WatchProperty(BackgroundVisualization->Settings->bVisible, 
 		[this](bool IsVisible) {
-			UpdateTriangleMaterialBasedOnBackground(IsVisible);
+			UpdateTriangleMaterialBasedOnDisplaySettings();			
 		});
 
 	BackgroundVisualization->OnBackgroundMaterialChange.AddWeakLambda(this,
@@ -403,6 +404,18 @@ void UUVEditorMode::Enter()
 
 	PropertyObjectsToTick.Add(BackgroundVisualization->Settings);
 
+	DistortionVisualization = NewObject<UUVEditorDistortionVisualization>(this);
+	DistortionVisualization->Targets = ToolInputObjects;
+	DistortionVisualization->Initialize();
+
+
+	DistortionVisualization->Settings->WatchProperty(DistortionVisualization->Settings->bVisible,
+		[this](bool IsVisible) {
+			UpdateTriangleMaterialBasedOnDisplaySettings();			
+		});
+
+	PropertyObjectsToTick.Add(DistortionVisualization->Settings);
+
 	UVEditorGridProperties = NewObject <UUVEditorGridProperties>(this);
 
 	UVEditorGridProperties->WatchProperty(UVEditorGridProperties->bDrawGrid,
@@ -642,6 +655,15 @@ UObject* UUVEditorMode::GetBackgroundSettingsObject()
 	return nullptr;
 }
 
+UObject* UUVEditorMode::GetDistortionVisualsSettingsObject()
+{
+	if (DistortionVisualization)
+	{
+		return DistortionVisualization->Settings;
+	}
+	return nullptr;
+}
+
 UObject* UUVEditorMode::GetGridSettingsObject()
 {
 	if (UVEditorGridProperties)
@@ -852,6 +874,8 @@ void UUVEditorMode::Exit()
 		BackgroundVisualization = nullptr;
 	}
 
+	DistortionVisualization = nullptr;
+
 	PropertyObjectsToTick.Empty();
 	LivePreviewWorld = nullptr;
 
@@ -1199,9 +1223,9 @@ void UUVEditorMode::UpdateActiveUDIMs()
 			Blocks.Reserve(ActiveUDIMs.Num());
 			for (FUDIMSpecifier& UDIMSpecifier : ActiveUDIMs)
 			{
-				Blocks.Add({ UDIMSpecifier.UDIM });
+				Blocks.Add({ UDIMSpecifier.UDIM, UDIMSpecifier.TextureResolution });
 				UDIMSpecifier.UCoord = Blocks.Last().BlockU();
-				UDIMSpecifier.VCoord = Blocks.Last().BlockV();
+				UDIMSpecifier.VCoord = Blocks.Last().BlockV();				
 			}
 		}
 		UVTool2DViewportAPI->SetUDIMBlocks(Blocks, true);
@@ -1219,6 +1243,17 @@ void UUVEditorMode::UpdateActiveUDIMs()
 		}
 		BackgroundVisualization->Settings->CheckAndUpdateWatched();
 	}
+
+	if (DistortionVisualization)
+	{
+		DistortionVisualization->Settings->PerUDIMTextureResolution.Empty();
+		DistortionVisualization->Settings->PerUDIMTextureResolution.Reserve(ActiveUDIMs.Num());
+		for (FUDIMSpecifier& UDIMSpecifier : ActiveUDIMs)
+		{
+			DistortionVisualization->Settings->PerUDIMTextureResolution.Add(UDIMSpecifier.UDIM, UDIMSpecifier.TextureResolution);
+		}
+		DistortionVisualization->Settings->CheckAndUpdateWatched();
+	}
 }
 
 void UUVEditorMode::PopulateUDIMsByAsset(int32 AssetId, TArray<FUDIMSpecifier>& UDIMsOut) const
@@ -1253,10 +1288,16 @@ void UUVEditorMode::PopulateUDIMsByAsset(int32 AssetId, TArray<FUDIMSpecifier>&
 }
 
 void UUVEditorMode::UpdateTriangleMaterialBasedOnBackground(bool IsBackgroundVisible)
+{
+	UpdateTriangleMaterialBasedOnDisplaySettings();
+}
+
+void UUVEditorMode::UpdateTriangleMaterialBasedOnDisplaySettings()
 {
 	float TriangleOpacity;
+
 	// We adjust the mesh opacity depending on whether we're layered over the background or not.
-	if (IsBackgroundVisible)
+	if (BackgroundVisualization && BackgroundVisualization->Settings && BackgroundVisualization->Settings->bVisible)
 	{
 		TriangleOpacity = FUVEditorUXSettings::UnwrapTriangleOpacityWithBackground;
 	}
@@ -1265,15 +1306,24 @@ void UUVEditorMode::UpdateTriangleMaterialBasedOnBackground(bool IsBackgroundVis
 		TriangleOpacity = FUVEditorUXSettings::UnwrapTriangleOpacity;
 	}
 
+	bool bUseVertexColors = DistortionVisualization && DistortionVisualization->Settings && DistortionVisualization->Settings->bVisible;
+
 	// Modify the material of the unwrapped mesh to account for the presence/absence of the background, 
 	// changing the opacity as set just above.
-	for (int32 AssetID = 0; AssetID < ToolInputObjects.Num(); ++AssetID) {
-		ToolInputObjects[AssetID]->UnwrapPreview->PreviewMesh->SetMaterial(
-			0, ToolSetupUtil::GetCustomTwoSidedDepthOffsetMaterial(
-				GetToolManager(),
-				FUVEditorUXSettings::GetTriangleColorByTargetIndex(AssetID),
-				FUVEditorUXSettings::UnwrapTriangleDepthOffset,
-				TriangleOpacity));
+	UMaterialInterface* Material = LoadObject<UMaterial>(nullptr, TEXT("/UVEditor/Materials/UVEditor_UnwrapMaterial"));
+	if (Material != nullptr)
+	{
+		for (int32 AssetID = 0; AssetID < ToolInputObjects.Num(); ++AssetID) {
+			UMaterialInstanceDynamic* MatInstance = UMaterialInstanceDynamic::Create(Material, GetToolManager());
+			MatInstance->SetVectorParameterValue(TEXT("Color"), FUVEditorUXSettings::GetTriangleColorByTargetIndex(AssetID));
+			MatInstance->SetScalarParameterValue(TEXT("PercentDepthOffset"), FUVEditorUXSettings::UnwrapTriangleDepthOffset);
+			MatInstance->SetScalarParameterValue(TEXT("Opacity"), TriangleOpacity);
+			MatInstance->SetScalarParameterValue(TEXT("UseVertexColors"), bUseVertexColors);
+
+			ToolInputObjects[AssetID]->bEnableTriangleVertexColors = bUseVertexColors;
+			ToolInputObjects[AssetID]->UnwrapPreview->PreviewMesh->SetMaterial(
+				0, MatInstance);
+		}
 	}
 }
 
@@ -1390,6 +1440,11 @@ void UUVEditorMode::ModeTick(float DeltaTime)
 		BackgroundVisualization->OnTick(DeltaTime);
 	}
 
+	if (DistortionVisualization)
+	{
+		DistortionVisualization->OnTick(DeltaTime);
+	}
+
 	for (int i = 0; i < ToolInputObjects.Num(); ++i)
 	{
 		TObjectPtr<UUVEditorToolMeshInput> ToolInput = ToolInputObjects[i];
diff --git a/Engine/Plugins/Experimental/UVEditor/Source/UVEditor/Private/UVEditorModeToolkit.cpp b/Engine/Plugins/Experimental/UVEditor/Source/UVEditor/Private/UVEditorModeToolkit.cpp
--- a/Engine/Plugins/Experimental/UVEditor/Source/UVEditor/Private/UVEditorModeToolkit.cpp
+++ b/Engine/Plugins/Experimental/UVEditor/Source/UVEditor/Private/UVEditorModeToolkit.cpp
@@ -299,6 +299,12 @@ TSharedRef<SWidget> FUVEditorModeToolkit::CreateBackgroundSettingsWidget()
 	return CreateDisplaySettingsWidget(Mode->GetBackgroundSettingsObject());
 }
 
+TSharedRef<SWidget> FUVEditorModeToolkit::CreateDistortionVisualsSettingsWidget()
+{
+	UUVEditorMode* Mode = Cast<UUVEditorMode>(GetScriptableEditorMode());
+	return CreateDisplaySettingsWidget(Mode->GetDistortionVisualsSettingsObject());
+}
+
 TSharedRef<SWidget> FUVEditorModeToolkit::CreateGridSettingsWidget()
 {
 	UUVEditorMode* Mode = Cast<UUVEditorMode>(GetScriptableEditorMode());
diff --git a/Engine/Plugins/Experimental/UVEditor/Source/UVEditor/Private/UVEditorToolkit.cpp b/Engine/Plugins/Experimental/UVEditor/Source/UVEditor/Private/UVEditorToolkit.cpp
--- a/Engine/Plugins/Experimental/UVEditor/Source/UVEditor/Private/UVEditorToolkit.cpp
+++ b/Engine/Plugins/Experimental/UVEditor/Source/UVEditor/Private/UVEditorToolkit.cpp
@@ -586,6 +586,17 @@ void FUVEditorToolkit::PostInitAssetEditor()
 				]
 				];
 
+			Container->AddSlot()
+				.AutoHeight()
+				.Padding(FMargin(0.f, 0.f, 8.f, 0.f))
+				[
+					SNew(SBox)
+					.MinDesiredWidth(500)
+				[
+					UVModeToolkit->CreateDistortionVisualsSettingsWidget()
+				]
+				];
+
 			TSharedRef<SWidget> Widget = SNew(SBorder)
 				.HAlign(HAlign_Fill)
 				.Padding(4)
diff --git a/Engine/Plugins/Experimental/UVEditor/Source/UVEditor/Public/UVEditorDistortionVisualization.h b/Engine/Plugins/Experimental/UVEditor/Source/UVEditor/Public/UVEditorDistortionVisualization.h
new file mode 100644
--- /dev/null
+++ b/Engine/Plugins/Experimental/UVEditor/Source/UVEditor/Public/UVEditorDistortionVisualization.h
@@ -0,0 +1,118 @@
+// Copyright Epic Games, Inc. All Rights Reserved.
+
+#pragma once
+
+#include "CoreMinimal.h"
+#include "Drawing/PreviewGeometryActor.h"
+#include "Drawing/TriangleSetComponent.h"
+#include "InteractiveTool.h"
+#include "GeometryBase.h"
+#include "ToolTargets/UVEditorToolMeshInput.h"
+
+#include "UVEditorDistortionVisualization.generated.h"
+
+UENUM()
+enum class EDistortionMetric : uint8
+{
+	// Represents the shape deviation from the worldspace triangle element, sans scaling. 
+	ReedBeta UMETA(DisplayName = "Elliptical Eccentricity"),
+	// Represents the root-mean-square strech over all directions vs the worldspace triangle element
+	Sander_L2 UMETA(DisplayName = "L2 Norm"),
+	// Represents the worst-case, or greatest stretch in any particular direction, vs the worldspace triangle element
+	Sander_LInf UMETA(DisplayName = "L Infinity Norm"),
+	// Shows the current texile density ratio, compared against the average baseline of the entire mesh. 
+	TexelDensity UMETA(DisplayName = "Texel Density")
+};
+
+/**
+ * Visualization settings for the UUVEditorDistortionVisualization
+ */
+UCLASS()
+class UVEDITOR_API UUVEditorDistortionVisualizationProperties : public UInteractiveToolPropertySet
+{
+	GENERATED_BODY()
+public:
+	/** Should the visualization be shown.*/
+	UPROPERTY(EditAnywhere, Category = "Distortion Visualization", meta = (DisplayName = "Display UV Distortion Visualization"))
+	bool bVisible = false;
+
+	/** The specific metric used to classify and visualize distortion. */
+	UPROPERTY(EditAnywhere, Category = "Distortion Visualization", meta = (DisplayName = "Metric"))
+	EDistortionMetric Metric = EDistortionMetric::ReedBeta;
+
+	//~ TODO - For now, we're disabling the user from changing this, as we're not going to support manual density targeting at the moment.
+	/** When visualizing texel density, use the average texel density, assuming a constant arbitrary texture resolution, over the entire model as the baseline. Otherwise, specify exact densities for a specific target texture size.*/
+	//~UPROPERTY(EditAnywhere, Category = "Distortion Visualization|Texel Density", meta = (DisplayName = "Use Average Baseline", EditCondition = "Metric == EDistortionMetric::TexelDensity"))
+
+	UPROPERTY()
+	bool bCompareToAverageDensity = true;
+
+	/** When computing average texel density, take into account specific UDIM tile resolutions for a more accurate representation */
+	UPROPERTY(EditAnywhere, Category = "Distortion Visualization|Texel Density", meta = (DisplayName = "Use UDIM Resolutions", EditCondition = "Metric == EDistortionMetric::TexelDensity"))
+	bool bRespectUDIMTextureResolutions = false;
+
+	//~ TODO - Enable support for manual density specificaitons at a future date
+	//~UPROPERTY(EditAnywhere, Category = "Distortion Visualization|Texel Density", meta = (DisplayName = "Texture Map Dimension", EditCondition = "!bCompareToAverageDensity && !bRespectUDIMTextureResolutions && Metric == EDistortionMetric::TexelDensity"))
+
+	UPROPERTY()
+	int32 MapSize;
+
+	//~ TODO - Enable support for manual density specificaitons at a future date	
+	//~UPROPERTY(EditAnywhere, Category = "Distortion Visualization|Texel Density", meta = (DisplayName = "Target Texel Density", EditCondition = "!bCompareToAverageDensity && Metric == EDistortionMetric::TexelDensity", ClampMin = 0.0, UIMin = 0.0))
+
+	UPROPERTY()
+	float TargetTexelDensity;
+
+	UPROPERTY()
+	TMap<int32, int32> PerUDIMTextureResolution;
+};
+
+/**
+  This class contains the logic for handling distortion visualization within the UV Editor. It manages the setting
+  of triangle color and computing the potential metrics used to describe the amount of distortion present in the UVs.
+ */
+UCLASS(Transient)
+class UVEDITOR_API UUVEditorDistortionVisualization : public UObject
+{
+	GENERATED_BODY()
+
+public:
+
+	/**
+	 * Client must call this after construction and parameter configuration to complete setup.
+	 */
+	void Initialize();
+
+	/**
+	 * Client must call this before destruction to clean up state.
+	 */
+	void Shutdown();
+
+	/**
+	 * Client must call this every frame for changes to .Settings to be reflected in rendered result.
+	 */
+	void OnTick(float DeltaTime);
+
+
+	
+public:
+
+	UPROPERTY()
+	TObjectPtr<UUVEditorDistortionVisualizationProperties> Settings;
+
+	UPROPERTY()
+	TArray<TObjectPtr<UUVEditorToolMeshInput>> Targets;
+
+protected:
+	bool bSettingsModified = false;
+
+	TArray<double> PerTargetAverageTexelDensity;
+	int32 MaximumTileTextureResolution;
+
+	void UpdateVisibility();
+	void ConfigureMeshColorsForTarget(int32 TargetIndex);
+	FColor GetDistortionColorForTriangle(const UE::Geometry::FDynamicMesh3& Mesh, int UVChannel, int Tid, double TargetAverageTexelDensity);
+	void ComputeInitialMeshSurfaceAreas();
+
+};
+
diff --git a/Engine/Plugins/Experimental/UVEditor/Source/UVEditor/Public/UVEditorMode.h b/Engine/Plugins/Experimental/UVEditor/Source/UVEditor/Public/UVEditorMode.h
--- a/Engine/Plugins/Experimental/UVEditor/Source/UVEditor/Public/UVEditorMode.h
+++ b/Engine/Plugins/Experimental/UVEditor/Source/UVEditor/Public/UVEditorMode.h
@@ -25,6 +25,7 @@ class UPreviewMesh;
 class UTexture2D;
 class UUVEditorToolMeshInput;
 class UUVEditorBackgroundPreview;
+class UUVEditorDistortionVisualization;
 class UUVToolAction;
 class UUVToolContextObject;
 class UUVToolSelectionAPI;
@@ -64,13 +65,17 @@ struct FUDIMSpecifier
 	UPROPERTY(VisibleAnywhere, Transient, Category = UDIM, meta = (DisplayName = "Block V Offset"))
 	int32 VCoord = 0;
 
+	UPROPERTY(EditAnywhere, Transient, Category = UDIM, meta = (DisplayName = "Block Texture Resolution"))
+	int32 TextureResolution = 512;
+
 	friend bool operator==(const FUDIMSpecifier& A, const FUDIMSpecifier& B )
 	{
-		return A.UDIM == B.UDIM;
+		return A.UDIM == B.UDIM && A.TextureResolution == B.TextureResolution;
 	}
 
 	friend uint32 GetTypeHash(const FUDIMSpecifier& UDIMSpecifier)
 	{
+		// We're only hashing on the UDIM index, as we don't want to ever have more than one entry with the same UDIM position
 		uint32 HashCode = GetTypeHash(UDIMSpecifier.UDIM);
 		return HashCode;
 	}
@@ -221,6 +226,9 @@ public:
 	/** @return A settings object suitable for display in a details panel to control the active tool's visualization settings. */
 	UObject* GetToolDisplaySettingsObject();
 
+	/** @return A settings object suitable for display in a details panel to control the UV distortion visualization settings. */
+	UObject* GetDistortionVisualsSettingsObject();
+
 	virtual void Render(IToolsContextRenderAPI* RenderAPI) ;
 	virtual void DrawHUD(FCanvas* Canvas, IToolsContextRenderAPI* RenderAPI);
 
@@ -241,10 +249,12 @@ public:
 	// We don't actually override MouseEnter, etc, because things get forwarded to the input
 	// router via FEditorModeTools, and we don't have any additional input handling to do at the mode level.
 
-	// Holds the background visualiztion
 	UPROPERTY()
 	TObjectPtr<UUVEditorBackgroundPreview> BackgroundVisualization;
 
+	UPROPERTY()
+	TObjectPtr<UUVEditorDistortionVisualization> DistortionVisualization;
+
 	// Hold a settings object to configure the grid
 	UPROPERTY()
 	TObjectPtr<UUVEditorGridProperties> UVEditorGridProperties = nullptr;
@@ -274,8 +284,10 @@ protected:
 	virtual void OnToolStarted(UInteractiveToolManager* Manager, UInteractiveTool* Tool) override;
 	virtual void OnToolEnded(UInteractiveToolManager* Manager, UInteractiveTool* Tool) override;
 	
+	UE_DEPRECATED(5.3, "UpdateTriangleMaterialBasedOnBackground is deprecated, replaced by the more generic UpdateTriangleMaterialBasedOnDisplaySettings.")
 	void UpdateTriangleMaterialBasedOnBackground(bool IsBackgroundVisible);
 	void UpdatePreviewMaterialBasedOnBackground();
+	void UpdateTriangleMaterialBasedOnDisplaySettings();
 
 	void UpdateActiveUDIMs();
 	int32 UDIMsChangedWatcherId;
diff --git a/Engine/Plugins/Experimental/UVEditor/Source/UVEditor/Public/UVEditorModeToolkit.h b/Engine/Plugins/Experimental/UVEditor/Source/UVEditor/Public/UVEditorModeToolkit.h
--- a/Engine/Plugins/Experimental/UVEditor/Source/UVEditor/Public/UVEditorModeToolkit.h
+++ b/Engine/Plugins/Experimental/UVEditor/Source/UVEditor/Public/UVEditorModeToolkit.h
@@ -44,6 +44,9 @@ public:
 	/** Creates submenu for manipulating the gizmo numerical UI */
 	virtual void MakeGizmoNumericalUISubMenu(FMenuBuilder& MenuBuilder);
 
+	/** Creates a widget where the distortion visualization settings can be changed. */
+	virtual TSharedRef<SWidget> CreateDistortionVisualsSettingsWidget();
+
 	/** Enables or diables the Play in Editor mode warning for the editor */
 	virtual void EnableShowPIEWarning(bool bEnable);
 
diff --git a/Engine/Plugins/Experimental/UVEditor/Source/UVEditorTools/Private/Math/UVMetrics.cpp b/Engine/Plugins/Experimental/UVEditor/Source/UVEditorTools/Private/Math/UVMetrics.cpp
new file mode 100644
--- /dev/null
+++ b/Engine/Plugins/Experimental/UVEditor/Source/UVEditorTools/Private/Math/UVMetrics.cpp
@@ -0,0 +1,151 @@
+// Copyright Epic Games, Inc. All Rights Reserved.
+
+#include "Math/UVMetrics.h"
+#include "DynamicMesh/DynamicMesh3.h"
+#include "DynamicMesh/DynamicMeshAttributeSet.h"
+
+using namespace UE::Geometry;
+
+namespace UVMetricsLocals
+{
+	double Triangle2DArea(FVector2f p1, FVector2f p2, FVector2f p3)
+	{
+		return FMath::Abs(((p2.X - p1.X) * (p3.Y - p1.Y) - (p3.X - p1.X) * (p2.Y - p1.Y)) / 2.0);
+	}
+
+	double Triangle3DArea(FVector3d q1, FVector3d q2, FVector3d q3)
+	{
+		FVector3d Q1 = FVector3d(q2.X - q1.X, q2.Y - q1.Y, q2.Z - q1.Z);
+		FVector3d Q2 = FVector3d(q3.X - q1.X, q3.Y - q1.Y, q3.Z - q1.Z);
+
+		FVector3d U = Q1.Cross(Q2);
+
+		return U.Length() / 2.0;
+	}
+
+	double Triangle2DPerimeter(FVector2f p1, FVector2f p2, FVector2f p3)
+	{
+		FVector2f Edge1 = p2 - p1;
+		FVector2f Edge2 = p3 - p1;
+		FVector2f Edge3 = p3 - p2;
+		return Edge1.Length() + Edge2.Length() + Edge3.Length();
+
+	}
+
+	double Triangle3DPerimeter(FVector3d q1, FVector3d q2, FVector3d q3)
+	{
+		FVector3d Edge1 = q2 - q1;
+		FVector3d Edge2 = q3 - q1;
+		FVector3d Edge3 = q3 - q2;
+		return Edge1.Length() + Edge2.Length() + Edge3.Length();
+	}
+
+	FVector2f GetTriangleSingularValues(FVector3d V1, FVector3d V2, FVector3d V3, FVector2f UV1, FVector2f UV2, FVector2f UV3)
+	{
+		FVector2f SingularValues;
+		FVector3d Ss, St;
+		double A;
+
+		A = Triangle2DArea(UV1, UV2, UV3);
+		Ss = (V1 * (UV2.Y - UV3.Y) + V2 * (UV3.Y - UV1.Y) + V3 * (UV1.Y - UV2.Y)) / A;
+		St = (V1 * (UV3.X - UV2.X) + V2 * (UV1.X - UV3.X) + V3 * (UV2.X - UV1.X)) / A;
+
+		double a, b, c;
+		a = Ss.Dot(Ss);
+		b = Ss.Dot(St);
+		c = St.Dot(St);
+
+		double Discriminant = FMath::Sqrt(FMath::Square(a - c) + 4 * FMath::Square(b));
+
+		SingularValues.X = FMath::Sqrt(0.5 * ((a + c) + Discriminant));
+		SingularValues.Y = FMath::Sqrt(0.5 * ((a + c) - Discriminant));
+
+		return SingularValues;
+	}
+}
+
+
+double FUVMetrics::ReedBeta(const FDynamicMesh3& Mesh, int32 UVChannel, int32 Tid)
+{
+	// Inspiration for the math used to compute the eccentricity comes from here:
+	// https://www.reedbeta.com/blog/conformal-texture-mapping/
+
+	FVector3d V1, V2, V3;
+	FVector2f UV1, UV2, UV3;
+	Mesh.GetTriVertices(Tid, V1, V2, V3);
+	Mesh.Attributes()->GetUVLayer(UVChannel)->GetTriElements(Tid, UV1, UV2, UV3);
+
+	FVector2f SingularValues = UVMetricsLocals::GetTriangleSingularValues(V1, V2, V3, UV1, UV2, UV3);
+	if (TMathUtil<double>::Abs(SingularValues.X) < TMathUtil<double>::Epsilon)
+	{
+		return 0;
+	}
+	else
+	{
+		double Eccentricity = FMath::Sqrt(1 - FMath::Square(SingularValues.Y / SingularValues.X));
+		return Eccentricity;
+	}
+}
+
+double FUVMetrics::Sander(const FDynamicMesh3& Mesh, int32 UVChannel, int32 Tid, bool bUseL2)
+{
+	// These metrics are defined by this paper:
+	// "Texture Mapping Progressive Meshes"
+	// https://dl.acm.org/doi/10.1145/383259.383307
+	// https://hhoppe.com/tmpm.pdf
+
+	FVector3d q1, q2, q3;
+	FVector2f p1, p2, p3;
+	Mesh.GetTriVertices(Tid, q1, q2, q3);
+	Mesh.Attributes()->GetUVLayer(UVChannel)->GetTriElements(Tid, p1, p2, p3);
+
+	double GammaMax, GammaMin, L2_Norm, L_Inf;
+	double L2_Norm_Normalized, L_Inf_Normalized;
+	{
+		FVector3d Ss, St;
+		double A, Aq;
+		A = UVMetricsLocals::Triangle2DArea(p1, p2, p3);
+		Aq = UVMetricsLocals::Triangle3DArea(q1, q2, q3);
+		Ss = (q1 * (p2.Y - p3.Y) + q2 * (p3.Y - p1.Y) + q3 * (p1.Y - p2.Y)) / (2 * A);
+		St = (q1 * (p3.X - p2.X) + q2 * (p1.X - p3.X) + q3 * (p2.X - p1.X)) / (2 * A);
+
+		double a, b, c;
+		a = Ss.Dot(Ss);
+		b = Ss.Dot(St);
+		c = St.Dot(St);
+
+		double D = FMath::Sqrt(FMath::Square(a - c) + 4 * FMath::Square(b));
+
+		GammaMax = FMath::Sqrt(0.5 * ((a + c) + D));
+		GammaMin = FMath::Sqrt(0.5 * ((a + c) - D));
+
+		L2_Norm = FMath::Sqrt(0.5 * (a + c));
+		L_Inf = GammaMax;
+
+		L2_Norm_Normalized = L2_Norm / FMath::Sqrt(Aq / A);
+		L_Inf_Normalized = L_Inf / FMath::Sqrt(Aq / A);
+	}
+	if (bUseL2)
+	{
+		return L2_Norm_Normalized;
+	}
+	else
+	{
+		return L_Inf_Normalized;
+	}
+}
+
+double FUVMetrics::TexelDensity(const FDynamicMesh3& Mesh, int32 UVChannel, int32 Tid, int32 MapSize)
+{
+	// Computes a metric describing the difference between the density of a single UV triangle relative to the desired density
+	FVector3d q1, q2, q3;
+	FVector2f p1, p2, p3;
+	Mesh.GetTriVertices(Tid, q1, q2, q3);
+	Mesh.Attributes()->GetUVLayer(UVChannel)->GetTriElements(Tid, p1, p2, p3);
+
+	double UVArea = UVMetricsLocals::Triangle2DArea(p1, p2, p3);
+	double WorldSpaceArea = UVMetricsLocals::Triangle3DArea(q1, q2, q3);
+	double Density = FMath::Sqrt(UVArea / WorldSpaceArea) * MapSize; // Since we're comparing area ratios here, we want to take the square root to get the linear scaling factor back out.
+																		// Technically Mapsize is Mapsize^2 in terms of pixels per unit area, so we implicitly take it's square root here.
+	return Density;
+}
diff --git a/Engine/Plugins/Experimental/UVEditor/Source/UVEditorTools/Private/ToolTargets/UVEditorToolMeshInput.cpp b/Engine/Plugins/Experimental/UVEditor/Source/UVEditorTools/Private/ToolTargets/UVEditorToolMeshInput.cpp
--- a/Engine/Plugins/Experimental/UVEditor/Source/UVEditorTools/Private/ToolTargets/UVEditorToolMeshInput.cpp
+++ b/Engine/Plugins/Experimental/UVEditor/Source/UVEditorTools/Private/ToolTargets/UVEditorToolMeshInput.cpp
@@ -24,10 +24,8 @@ namespace UVEditorToolMeshInputLocals
 
 	void NotifyUnwrapPreviewDeferredEditCompleted(UUVEditorToolMeshInput* InputObject,
 		const TArray<int32>* ChangedVids, const TArray<int32>* ChangedConnectivityTids, 
-		const TArray<int32>* FastRenderUpdateTids)
+		const TArray<int32>* FastRenderUpdateTids, EMeshRenderAttributeFlags ChangedAttribs)
 	{
-		EMeshRenderAttributeFlags ChangedAttribs = EMeshRenderAttributeFlags::Positions | EMeshRenderAttributeFlags::VertexUVs;
-
 		// Currently there's no way to do a partial or region update if the connectivity changed in any way.
 		if (ChangedConnectivityTids == nullptr || !ChangedConnectivityTids->IsEmpty())
 		{
@@ -44,10 +42,8 @@ namespace UVEditorToolMeshInputLocals
 	}
 
 	void NotifyAppliedPreviewDeferredEditCompleted(UUVEditorToolMeshInput* InputObject,
-		const TArray<int32>* ChangedVids, const TArray<int32>* FastRenderUpdateTids)
+		const TArray<int32>* ChangedVids, const TArray<int32>* FastRenderUpdateTids, EMeshRenderAttributeFlags ChangedAttribs)
 	{
-		EMeshRenderAttributeFlags ChangedAttribs = EMeshRenderAttributeFlags::VertexUVs;
-
 		if (FastRenderUpdateTids) 
 		{
 			InputObject->AppliedPreview->PreviewMesh->NotifyRegionDeferredEditCompleted(*FastRenderUpdateTids, ChangedAttribs);
@@ -159,7 +155,7 @@ void UUVEditorToolMeshInput::UpdateUnwrapPreviewOverlayFromPositions(const TArra
 
 	// We only updated UV's, but this update usually comes after a change to unwrap vert positions, and we
 	// can assume that the user hasn't issued a notification yet. So we do a normal unwrap preview update.
-	NotifyUnwrapPreviewDeferredEditCompleted(this, ChangedVids, ChangedConnectivityTids, FastRenderUpdateTids);
+	NotifyUnwrapPreviewDeferredEditCompleted(this, ChangedVids, ChangedConnectivityTids, FastRenderUpdateTids, GetUnwrapUpdateAttributeFlags());
 
 	if (WireframeDisplay)
 	{
@@ -192,7 +188,7 @@ void UUVEditorToolMeshInput::UpdateAppliedPreviewFromUnwrapPreview(const TArray<
 			UVUnwrapMeshUtil::UpdateOverlayFromOverlay(*SourceOverlay, *DestOverlay, false, ChangedVids, ChangedConnectivityTids);
 	}, false);
 
-	NotifyAppliedPreviewDeferredEditCompleted(this, ChangedVids, FastRenderUpdateTids);
+	NotifyAppliedPreviewDeferredEditCompleted(this, ChangedVids, FastRenderUpdateTids, GetAppliedUpdateAttributeFlags());
 }
 
 void UUVEditorToolMeshInput::UpdateUnwrapPreviewFromAppliedPreview(const TArray<int32>* ChangedElementIDs, 
@@ -211,7 +207,7 @@ void UUVEditorToolMeshInput::UpdateUnwrapPreviewFromAppliedPreview(const TArray<
 		UVUnwrapMeshUtil::UpdateOverlayFromOverlay(*SourceOverlay, *DestOverlay, false, ChangedElementIDs, ChangedConnectivityTids);
 	}, false);
 
-	NotifyUnwrapPreviewDeferredEditCompleted(this, ChangedElementIDs, ChangedConnectivityTids, FastRenderUpdateTids);
+	NotifyUnwrapPreviewDeferredEditCompleted(this, ChangedElementIDs, ChangedConnectivityTids, FastRenderUpdateTids, GetUnwrapUpdateAttributeFlags());
 
 	if (WireframeDisplay)
 	{
@@ -245,7 +241,7 @@ void UUVEditorToolMeshInput::UpdatePreviewsFromCanonical(const TArray<int32>* Ch
 	{
 		UVUnwrapMeshUtil::UpdateUVUnwrapMesh(*UnwrapCanonical, Mesh, ChangedVids, ChangedConnectivityTids);
 	}, false);
-	NotifyUnwrapPreviewDeferredEditCompleted(this, ChangedVids, ChangedConnectivityTids, FastRenderUpdateTids);
+	NotifyUnwrapPreviewDeferredEditCompleted(this, ChangedVids, ChangedConnectivityTids, FastRenderUpdateTids, GetUnwrapUpdateAttributeFlags());
 
 	if (WireframeDisplay)
 	{
@@ -259,7 +255,7 @@ void UUVEditorToolMeshInput::UpdatePreviewsFromCanonical(const TArray<int32>* Ch
 		FDynamicMeshUVOverlay* DestOverlay = Mesh.Attributes()->GetUVLayer(UVLayerIndex);
 		UVUnwrapMeshUtil::UpdateOverlayFromOverlay(*SourceOverlay, *DestOverlay, true, ChangedVids, ChangedConnectivityTids);
 	}, false);
-	NotifyAppliedPreviewDeferredEditCompleted(this, ChangedVids, FastRenderUpdateTids);
+	NotifyAppliedPreviewDeferredEditCompleted(this, ChangedVids, FastRenderUpdateTids, GetAppliedUpdateAttributeFlags());
 }
 
 void UUVEditorToolMeshInput::UpdateAllFromUnwrapPreview(const TArray<int32>* ChangedVids, 
@@ -355,3 +351,20 @@ void UUVEditorToolMeshInput::AppliedVidToUnwrapVids(int32 AppliedVid, TArray<int
 	const FDynamicMeshUVOverlay* CanonicalOverlay = AppliedCanonical->Attributes()->GetUVLayer(UVLayerIndex);
 	CanonicalOverlay->GetVertexElements(AppliedVid, UnwrapVidsOut);
 }
+
+EMeshRenderAttributeFlags UUVEditorToolMeshInput::GetUnwrapUpdateAttributeFlags() const
+{
+	if (bEnableTriangleVertexColors)
+	{
+		return EMeshRenderAttributeFlags::Positions | EMeshRenderAttributeFlags::VertexUVs | EMeshRenderAttributeFlags::VertexColors;
+	}
+	else
+	{
+		return EMeshRenderAttributeFlags::Positions | EMeshRenderAttributeFlags::VertexUVs;
+	}
+}
+
+EMeshRenderAttributeFlags UUVEditorToolMeshInput::GetAppliedUpdateAttributeFlags() const
+{
+	return EMeshRenderAttributeFlags::VertexUVs;
+}
\ No newline at end of file
diff --git a/Engine/Plugins/Experimental/UVEditor/Source/UVEditorTools/Private/UVEditorUXSettings.cpp b/Engine/Plugins/Experimental/UVEditor/Source/UVEditorTools/Private/UVEditorUXSettings.cpp
--- a/Engine/Plugins/Experimental/UVEditor/Source/UVEditorTools/Private/UVEditorUXSettings.cpp
+++ b/Engine/Plugins/Experimental/UVEditor/Source/UVEditorTools/Private/UVEditorUXSettings.cpp
@@ -128,6 +128,140 @@ FLinearColor FUVEditorUXSettings::GetBoundaryColorByTargetIndex(int32 TargetInde
 	return BoundaryColorHSV.HSVToLinearRGB();
 }
 
+FColor FUVEditorUXSettings::MakeCividisColorFromScalar(float Scalar)
+{
+	// Color map sourced from:
+	// "Optimizing colormaps with consideration for color vision deficiency to enable accurate interpretation of scientific data"
+	// https://doi.org/10.1371/journal.pone.0199239
+
+	float RComponents[256] = { 0.0000, 0.0000, 0.0000, 0.0000, 0.0000, 0.0000, 0.0000, 0.0000, 0.0000, 0.0000, 0.0000, 0.0000, 0.0000,
+		0.0000, 0.0000, 0.0000, 0.0000, 0.0000, 0.0000, 0.0000, 0.0000, 0.0000, 0.0000, 0.0000, 0.0000, 0.0000, 0.0000, 0.0000, 0.0000,
+		0.0000, 0.0000, 0.0055, 0.0236, 0.0416, 0.0576, 0.0710, 0.0827, 0.0932, 0.1030, 0.1120, 0.1204, 0.1283, 0.1359, 0.1431, 0.1500,
+		0.1566, 0.1630, 0.1692, 0.1752, 0.1811, 0.1868, 0.1923, 0.1977, 0.2030, 0.2082, 0.2133, 0.2183, 0.2232, 0.2281, 0.2328, 0.2375,
+		0.2421, 0.2466, 0.2511, 0.2556, 0.2599, 0.2643, 0.2686, 0.2728, 0.2770, 0.2811, 0.2853, 0.2894, 0.2934, 0.2974, 0.3014, 0.3054,
+		0.3093, 0.3132, 0.3170, 0.3209, 0.3247, 0.3285, 0.3323, 0.3361, 0.3398, 0.3435, 0.3472, 0.3509, 0.3546, 0.3582, 0.3619, 0.3655,
+		0.3691, 0.3727, 0.3763, 0.3798, 0.3834, 0.3869, 0.3905, 0.3940, 0.3975, 0.4010, 0.4045, 0.4080, 0.4114, 0.4149, 0.4183, 0.4218,
+		0.4252, 0.4286, 0.4320, 0.4354, 0.4388, 0.4422, 0.4456, 0.4489, 0.4523, 0.4556, 0.4589, 0.4622, 0.4656, 0.4689, 0.4722, 0.4756,
+		0.4790, 0.4825, 0.4861, 0.4897, 0.4934, 0.4971, 0.5008, 0.5045, 0.5083, 0.5121, 0.5158, 0.5196, 0.5234, 0.5272, 0.5310, 0.5349,
+		0.5387, 0.5425, 0.5464, 0.5502, 0.5541, 0.5579, 0.5618, 0.5657, 0.5696, 0.5735, 0.5774, 0.5813, 0.5852, 0.5892, 0.5931, 0.5970,
+		0.6010, 0.6050, 0.6089, 0.6129, 0.6168, 0.6208, 0.6248, 0.6288, 0.6328, 0.6368, 0.6408, 0.6449, 0.6489, 0.6529, 0.6570, 0.6610,
+		0.6651, 0.6691, 0.6732, 0.6773, 0.6813, 0.6854, 0.6895, 0.6936, 0.6977, 0.7018, 0.7060, 0.7101, 0.7142, 0.7184, 0.7225, 0.7267,
+		0.7308, 0.7350, 0.7392, 0.7434, 0.7476, 0.7518, 0.7560, 0.7602, 0.7644, 0.7686, 0.7729, 0.7771, 0.7814, 0.7856, 0.7899, 0.7942,
+		0.7985, 0.8027, 0.8070, 0.8114, 0.8157, 0.8200, 0.8243, 0.8287, 0.8330, 0.8374, 0.8417, 0.8461, 0.8505, 0.8548, 0.8592, 0.8636,
+		0.8681, 0.8725, 0.8769, 0.8813, 0.8858, 0.8902, 0.8947, 0.8992, 0.9037, 0.9082, 0.9127, 0.9172, 0.9217, 0.9262, 0.9308, 0.9353,
+		0.9399, 0.9444, 0.9490, 0.9536, 0.9582, 0.9628, 0.9674, 0.9721, 0.9767, 0.9814, 0.9860, 0.9907, 0.9954, 1.0000, 1.0000, 1.0000,
+		1.0000, 1.0000, 1.0000};
+	float GComponents[256] = { 0.1262, 0.1292, 0.1321, 0.1350, 0.1379, 0.1408, 0.1437, 0.1465, 0.1492, 0.1519, 0.1546, 0.1574, 0.1601,
+		0.1629, 0.1657, 0.1685, 0.1714, 0.1743, 0.1773, 0.1798, 0.1817, 0.1834, 0.1852, 0.1872, 0.1901, 0.1930, 0.1958, 0.1987, 0.2015,
+		0.2044, 0.2073, 0.2101, 0.2130, 0.2158, 0.2187, 0.2215, 0.2244, 0.2272, 0.2300, 0.2329, 0.2357, 0.2385, 0.2414, 0.2442, 0.2470,
+		0.2498, 0.2526, 0.2555, 0.2583, 0.2611, 0.2639, 0.2667, 0.2695, 0.2723, 0.2751, 0.2780, 0.2808, 0.2836, 0.2864, 0.2892, 0.2920,
+		0.2948, 0.2976, 0.3004, 0.3032, 0.3060, 0.3088, 0.3116, 0.3144, 0.3172, 0.3200, 0.3228, 0.3256, 0.3284, 0.3312, 0.3340, 0.3368,
+		0.3396, 0.3424, 0.3453, 0.3481, 0.3509, 0.3537, 0.3565, 0.3593, 0.3622, 0.3650, 0.3678, 0.3706, 0.3734, 0.3763, 0.3791, 0.3819,
+		0.3848, 0.3876, 0.3904, 0.3933, 0.3961, 0.3990, 0.4018, 0.4047, 0.4075, 0.4104, 0.4132, 0.4161, 0.4189, 0.4218, 0.4247, 0.4275,
+		0.4304, 0.4333, 0.4362, 0.4390, 0.4419, 0.4448, 0.4477, 0.4506, 0.4535, 0.4564, 0.4593, 0.4622, 0.4651, 0.4680, 0.4709, 0.4738,
+		0.4767, 0.4797, 0.4826, 0.4856, 0.4886, 0.4915, 0.4945, 0.4975, 0.5005, 0.5035, 0.5065, 0.5095, 0.5125, 0.5155, 0.5186, 0.5216,
+		0.5246, 0.5277, 0.5307, 0.5338, 0.5368, 0.5399, 0.5430, 0.5461, 0.5491, 0.5522, 0.5553, 0.5584, 0.5615, 0.5646, 0.5678, 0.5709,
+		0.5740, 0.5772, 0.5803, 0.5835, 0.5866, 0.5898, 0.5929, 0.5961, 0.5993, 0.6025, 0.6057, 0.6089, 0.6121, 0.6153, 0.6185, 0.6217,
+		0.6250, 0.6282, 0.6315, 0.6347, 0.6380, 0.6412, 0.6445, 0.6478, 0.6511, 0.6544, 0.6577, 0.6610, 0.6643, 0.6676, 0.6710, 0.6743,
+		0.6776, 0.6810, 0.6844, 0.6877, 0.6911, 0.6945, 0.6979, 0.7013, 0.7047, 0.7081, 0.7115, 0.7150, 0.7184, 0.7218, 0.7253, 0.7288,
+		0.7322, 0.7357, 0.7392, 0.7427, 0.7462, 0.7497, 0.7532, 0.7568, 0.7603, 0.7639, 0.7674, 0.7710, 0.7745, 0.7781, 0.7817, 0.7853,
+		0.7889, 0.7926, 0.7962, 0.7998, 0.8035, 0.8071, 0.8108, 0.8145, 0.8182, 0.8219, 0.8256, 0.8293, 0.8330, 0.8367, 0.8405, 0.8442,
+		0.8480, 0.8518, 0.8556, 0.8593, 0.8632, 0.8670, 0.8708, 0.8746, 0.8785, 0.8823, 0.8862, 0.8901, 0.8940, 0.8979, 0.9018, 0.9057,
+		0.9094, 0.9131, 0.9169};
+	float BComponents[256] = { 0.3015, 0.3077, 0.3142, 0.3205, 0.3269, 0.3334, 0.3400, 0.3467, 0.3537, 0.3606, 0.3676, 0.3746, 0.3817,
+		0.3888, 0.3960, 0.4031, 0.4102, 0.4172, 0.4241, 0.4307, 0.4347, 0.4363, 0.4368, 0.4368, 0.4365, 0.4361, 0.4356, 0.4349, 0.4343,
+		0.4336, 0.4329, 0.4322, 0.4314, 0.4308, 0.4301, 0.4293, 0.4287, 0.4280, 0.4274, 0.4268, 0.4262, 0.4256, 0.4251, 0.4245, 0.4241,
+		0.4236, 0.4232, 0.4228, 0.4224, 0.4220, 0.4217, 0.4214, 0.4212, 0.4209, 0.4207, 0.4205, 0.4204, 0.4203, 0.4202, 0.4201, 0.4200,
+		0.4200, 0.4200, 0.4201, 0.4201, 0.4202, 0.4203, 0.4205, 0.4206, 0.4208, 0.4210, 0.4212, 0.4215, 0.4218, 0.4221, 0.4224, 0.4227,
+		0.4231, 0.4236, 0.4240, 0.4244, 0.4249, 0.4254, 0.4259, 0.4264, 0.4270, 0.4276, 0.4282, 0.4288, 0.4294, 0.4302, 0.4308, 0.4316,
+		0.4322, 0.4331, 0.4338, 0.4346, 0.4355, 0.4364, 0.4372, 0.4381, 0.4390, 0.4400, 0.4409, 0.4419, 0.4430, 0.4440, 0.4450, 0.4462,
+		0.4473, 0.4485, 0.4496, 0.4508, 0.4521, 0.4534, 0.4547, 0.4561, 0.4575, 0.4589, 0.4604, 0.4620, 0.4635, 0.4650, 0.4665, 0.4679,
+		0.4691, 0.4701, 0.4707, 0.4714, 0.4719, 0.4723, 0.4727, 0.4730, 0.4732, 0.4734, 0.4736, 0.4737, 0.4738, 0.4739, 0.4739, 0.4738,
+		0.4739, 0.4738, 0.4736, 0.4735, 0.4733, 0.4732, 0.4729, 0.4727, 0.4723, 0.4720, 0.4717, 0.4714, 0.4709, 0.4705, 0.4701, 0.4696,
+		0.4691, 0.4685, 0.4680, 0.4673, 0.4668, 0.4662, 0.4655, 0.4649, 0.4641, 0.4632, 0.4625, 0.4617, 0.4609, 0.4600, 0.4591, 0.4583,
+		0.4573, 0.4562, 0.4553, 0.4543, 0.4532, 0.4521, 0.4511, 0.4499, 0.4487, 0.4475, 0.4463, 0.4450, 0.4437, 0.4424, 0.4409, 0.4396,
+		0.4382, 0.4368, 0.4352, 0.4338, 0.4322, 0.4307, 0.4290, 0.4273, 0.4258, 0.4241, 0.4223, 0.4205, 0.4188, 0.4168, 0.4150, 0.4129, 
+		0.4111, 0.4090, 0.4070, 0.4049, 0.4028, 0.4007, 0.3984, 0.3961, 0.3938, 0.3915, 0.3892, 0.3869, 0.3843, 0.3818, 0.3793, 0.3766,
+		0.3739, 0.3712, 0.3684, 0.3657, 0.3627, 0.3599, 0.3569, 0.3538, 0.3507, 0.3474, 0.3442, 0.3409, 0.3374, 0.3340, 0.3306, 0.3268,
+		0.3232, 0.3195, 0.3155, 0.3116, 0.3076, 0.3034, 0.2990, 0.2947, 0.2901, 0.2856, 0.2807, 0.2759, 0.2708, 0.2655, 0.2600, 0.2593,
+		0.2634, 0.2680, 0.2731};
+
+	const float RedSclr = FMath::Clamp<float>(Scalar, 0.f, 1.f);
+	const float GreenSclr = FMath::Clamp<float>(Scalar, 0.f, 1.f);
+	const float BlueSclr = FMath::Clamp<float>(Scalar, 0.f, 1.f);
+	const uint8 R = (uint8)FMath::TruncToInt(254 * RComponents[(uint8)FMath::TruncToInt(254 * RedSclr)]);
+	const uint8 G = (uint8)FMath::TruncToInt(254 * GComponents[(uint8)FMath::TruncToInt(254 * GreenSclr)]);
+	const uint8 B = (uint8)FMath::TruncToInt(254 * BComponents[(uint8)FMath::TruncToInt(254 * BlueSclr)]);
+	return FColor(R, G, B);
+}
+
+FColor FUVEditorUXSettings::MakeTurboColorFromScalar(float Scalar)
+{
+	// Color map sourced from:
+	// https://ai.googleblog.com/2019/08/turbo-improved-rainbow-colormap-for.html
+
+	float TurboSRGB[256][3] = { {0.18995,0.07176,0.23217},{0.19483,0.08339,0.26149},{0.19956,0.09498,0.29024},{0.20415,0.10652,0.31844},
+		{0.20860,0.11802,0.34607},{0.21291,0.12947,0.37314},{0.21708,0.14087,0.39964},{0.22111,0.15223,0.42558},{0.22500,0.16354,0.45096},
+		{0.22875,0.17481,0.47578},{0.23236,0.18603,0.50004},{0.23582,0.19720,0.52373},{0.23915,0.20833,0.54686},{0.24234,0.21941,0.56942},
+		{0.24539,0.23044,0.59142},{0.24830,0.24143,0.61286},{0.25107,0.25237,0.63374},{0.25369,0.26327,0.65406},{0.25618,0.27412,0.67381},
+		{0.25853,0.28492,0.69300},{0.26074,0.29568,0.71162},{0.26280,0.30639,0.72968},{0.26473,0.31706,0.74718},{0.26652,0.32768,0.76412},
+		{0.26816,0.33825,0.78050},{0.26967,0.34878,0.79631},{0.27103,0.35926,0.81156},{0.27226,0.36970,0.82624},{0.27334,0.38008,0.84037},
+		{0.27429,0.39043,0.85393},{0.27509,0.40072,0.86692},{0.27576,0.41097,0.87936},{0.27628,0.42118,0.89123},{0.27667,0.43134,0.90254},
+		{0.27691,0.44145,0.91328},{0.27701,0.45152,0.92347},{0.27698,0.46153,0.93309},{0.27680,0.47151,0.94214},{0.27648,0.48144,0.95064},
+		{0.27603,0.49132,0.95857},{0.27543,0.50115,0.96594},{0.27469,0.51094,0.97275},{0.27381,0.52069,0.97899},{0.27273,0.53040,0.98461},
+		{0.27106,0.54015,0.98930},{0.26878,0.54995,0.99303},{0.26592,0.55979,0.99583},{0.26252,0.56967,0.99773},{0.25862,0.57958,0.99876},
+		{0.25425,0.58950,0.99896},{0.24946,0.59943,0.99835},{0.24427,0.60937,0.99697},{0.23874,0.61931,0.99485},{0.23288,0.62923,0.99202},
+		{0.22676,0.63913,0.98851},{0.22039,0.64901,0.98436},{0.21382,0.65886,0.97959},{0.20708,0.66866,0.97423},{0.20021,0.67842,0.96833},
+		{0.19326,0.68812,0.96190},{0.18625,0.69775,0.95498},{0.17923,0.70732,0.94761},{0.17223,0.71680,0.93981},{0.16529,0.72620,0.93161},
+		{0.15844,0.73551,0.92305},{0.15173,0.74472,0.91416},{0.14519,0.75381,0.90496},{0.13886,0.76279,0.89550},{0.13278,0.77165,0.88580},
+		{0.12698,0.78037,0.87590},{0.12151,0.78896,0.86581},{0.11639,0.79740,0.85559},{0.11167,0.80569,0.84525},{0.10738,0.81381,0.83484},
+		{0.10357,0.82177,0.82437},{0.10026,0.82955,0.81389},{0.09750,0.83714,0.80342},{0.09532,0.84455,0.79299},{0.09377,0.85175,0.78264},
+		{0.09287,0.85875,0.77240},{0.09267,0.86554,0.76230},{0.09320,0.87211,0.75237},{0.09451,0.87844,0.74265},{0.09662,0.88454,0.73316},
+		{0.09958,0.89040,0.72393},{0.10342,0.89600,0.71500},{0.10815,0.90142,0.70599},{0.11374,0.90673,0.69651},{0.12014,0.91193,0.68660},
+		{0.12733,0.91701,0.67627},{0.13526,0.92197,0.66556},{0.14391,0.92680,0.65448},{0.15323,0.93151,0.64308},{0.16319,0.93609,0.63137},
+		{0.17377,0.94053,0.61938},{0.18491,0.94484,0.60713},{0.19659,0.94901,0.59466},{0.20877,0.95304,0.58199},{0.22142,0.95692,0.56914},
+		{0.23449,0.96065,0.55614},{0.24797,0.96423,0.54303},{0.26180,0.96765,0.52981},{0.27597,0.97092,0.51653},{0.29042,0.97403,0.50321},
+		{0.30513,0.97697,0.48987},{0.32006,0.97974,0.47654},{0.33517,0.98234,0.46325},{0.35043,0.98477,0.45002},{0.36581,0.98702,0.43688},
+		{0.38127,0.98909,0.42386},{0.39678,0.99098,0.41098},{0.41229,0.99268,0.39826},{0.42778,0.99419,0.38575},{0.44321,0.99551,0.37345},
+		{0.45854,0.99663,0.36140},{0.47375,0.99755,0.34963},{0.48879,0.99828,0.33816},{0.50362,0.99879,0.32701},{0.51822,0.99910,0.31622},
+		{0.53255,0.99919,0.30581},{0.54658,0.99907,0.29581},{0.56026,0.99873,0.28623},{0.57357,0.99817,0.27712},{0.58646,0.99739,0.26849},
+		{0.59891,0.99638,0.26038},{0.61088,0.99514,0.25280},{0.62233,0.99366,0.24579},{0.63323,0.99195,0.23937},{0.64362,0.98999,0.23356},
+		{0.65394,0.98775,0.22835},{0.66428,0.98524,0.22370},{0.67462,0.98246,0.21960},{0.68494,0.97941,0.21602},{0.69525,0.97610,0.21294},
+		{0.70553,0.97255,0.21032},{0.71577,0.96875,0.20815},{0.72596,0.96470,0.20640},{0.73610,0.96043,0.20504},{0.74617,0.95593,0.20406},
+		{0.75617,0.95121,0.20343},{0.76608,0.94627,0.20311},{0.77591,0.94113,0.20310},{0.78563,0.93579,0.20336},{0.79524,0.93025,0.20386},
+		{0.80473,0.92452,0.20459},{0.81410,0.91861,0.20552},{0.82333,0.91253,0.20663},{0.83241,0.90627,0.20788},{0.84133,0.89986,0.20926},
+		{0.85010,0.89328,0.21074},{0.85868,0.88655,0.21230},{0.86709,0.87968,0.21391},{0.87530,0.87267,0.21555},{0.88331,0.86553,0.21719},
+		{0.89112,0.85826,0.21880},{0.89870,0.85087,0.22038},{0.90605,0.84337,0.22188},{0.91317,0.83576,0.22328},{0.92004,0.82806,0.22456},
+		{0.92666,0.82025,0.22570},{0.93301,0.81236,0.22667},{0.93909,0.80439,0.22744},{0.94489,0.79634,0.22800},{0.95039,0.78823,0.22831},
+		{0.95560,0.78005,0.22836},{0.96049,0.77181,0.22811},{0.96507,0.76352,0.22754},{0.96931,0.75519,0.22663},{0.97323,0.74682,0.22536},
+		{0.97679,0.73842,0.22369},{0.98000,0.73000,0.22161},{0.98289,0.72140,0.21918},{0.98549,0.71250,0.21650},{0.98781,0.70330,0.21358},
+		{0.98986,0.69382,0.21043},{0.99163,0.68408,0.20706},{0.99314,0.67408,0.20348},{0.99438,0.66386,0.19971},{0.99535,0.65341,0.19577},
+		{0.99607,0.64277,0.19165},{0.99654,0.63193,0.18738},{0.99675,0.62093,0.18297},{0.99672,0.60977,0.17842},{0.99644,0.59846,0.17376},
+		{0.99593,0.58703,0.16899},{0.99517,0.57549,0.16412},{0.99419,0.56386,0.15918},{0.99297,0.55214,0.15417},{0.99153,0.54036,0.14910},
+		{0.98987,0.52854,0.14398},{0.98799,0.51667,0.13883},{0.98590,0.50479,0.13367},{0.98360,0.49291,0.12849},{0.98108,0.48104,0.12332},
+		{0.97837,0.46920,0.11817},{0.97545,0.45740,0.11305},{0.97234,0.44565,0.10797},{0.96904,0.43399,0.10294},{0.96555,0.42241,0.09798},
+		{0.96187,0.41093,0.09310},{0.95801,0.39958,0.08831},{0.95398,0.38836,0.08362},{0.94977,0.37729,0.07905},{0.94538,0.36638,0.07461},
+		{0.94084,0.35566,0.07031},{0.93612,0.34513,0.06616},{0.93125,0.33482,0.06218},{0.92623,0.32473,0.05837},{0.92105,0.31489,0.05475},
+		{0.91572,0.30530,0.05134},{0.91024,0.29599,0.04814},{0.90463,0.28696,0.04516},{0.89888,0.27824,0.04243},{0.89298,0.26981,0.03993},
+		{0.88691,0.26152,0.03753},{0.88066,0.25334,0.03521},{0.87422,0.24526,0.03297},{0.86760,0.23730,0.03082},{0.86079,0.22945,0.02875},
+		{0.85380,0.22170,0.02677},{0.84662,0.21407,0.02487},{0.83926,0.20654,0.02305},{0.83172,0.19912,0.02131},{0.82399,0.19182,0.01966},
+		{0.81608,0.18462,0.01809},{0.80799,0.17753,0.01660},{0.79971,0.17055,0.01520},{0.79125,0.16368,0.01387},{0.78260,0.15693,0.01264},
+		{0.77377,0.15028,0.01148},{0.76476,0.14374,0.01041},{0.75556,0.13731,0.00942},{0.74617,0.13098,0.00851},{0.73661,0.12477,0.00769},
+		{0.72686,0.11867,0.00695},{0.71692,0.11268,0.00629},{0.70680,0.10680,0.00571},{0.69650,0.10102,0.00522},{0.68602,0.09536,0.00481},
+		{0.67535,0.08980,0.00449},{0.66449,0.08436,0.00424},{0.65345,0.07902,0.00408},{0.64223,0.07380,0.00401},{0.63082,0.06868,0.00401},
+		{0.61923,0.06367,0.00410},{0.60746,0.05878,0.00427},{0.59550,0.05399,0.00453},{0.58336,0.04931,0.00486},{0.57103,0.04474,0.00529},
+		{0.55852,0.04028,0.00579},{0.54583,0.03593,0.00638},{0.53295,0.03169,0.00705},{0.51989,0.02756,0.00780},{0.50664,0.02354,0.00863},
+		{0.49321,0.01963,0.00955},{0.47960,0.01583,0.01055} };
+
+	const float ClampedScalar = FMath::Clamp<float>(Scalar, -0.5f, 0.5f);
+	const uint8 RampIndex = (uint8)FMath::TruncToInt(255 * (ClampedScalar + 0.5f));
+	const uint8 R = (uint8)FMath::TruncToInt(255 * TurboSRGB[RampIndex][0]);
+	const uint8 G = (uint8)FMath::TruncToInt(255 * TurboSRGB[RampIndex][1]);
+	const uint8 B = (uint8)FMath::TruncToInt(255 * TurboSRGB[RampIndex][2]);
+	FLinearColor LinearColor = FLinearColor::FromSRGBColor(FColor(R, G, B));
+	return LinearColor.ToFColor(false);
+}
+
 // The mappings here depend on the way we view the unwrap world, set up in UVEditorToolkit.cpp.
 // Currently, we look at the world with Z pointing towards the camera, the positive X axis pointing right,
 // and the positive Y axis pointing down. This means that we have to flip the V coordinate to map it
diff --git a/Engine/Plugins/Experimental/UVEditor/Source/UVEditorTools/Public/ContextObjects/UVToolContextObjects.h b/Engine/Plugins/Experimental/UVEditor/Source/UVEditorTools/Public/ContextObjects/UVToolContextObjects.h
--- a/Engine/Plugins/Experimental/UVEditor/Source/UVEditorTools/Public/ContextObjects/UVToolContextObjects.h
+++ b/Engine/Plugins/Experimental/UVEditor/Source/UVEditorTools/Public/ContextObjects/UVToolContextObjects.h
@@ -163,6 +163,9 @@ struct UVEDITORTOOLS_API FUDIMBlock
 	UPROPERTY()
 	int32 UDIM = 1001;
 
+	UPROPERTY()
+	int32 TextureResolution = 512;
+
 	int32 BlockU() const;
 	int32 BlockV() const;
 	void SetFromBlocks(int32 BlockU, int32 BlockV);
diff --git a/Engine/Plugins/Experimental/UVEditor/Source/UVEditorTools/Public/Math/UVMetrics.h b/Engine/Plugins/Experimental/UVEditor/Source/UVEditorTools/Public/Math/UVMetrics.h
new file mode 100644
--- /dev/null
+++ b/Engine/Plugins/Experimental/UVEditor/Source/UVEditorTools/Public/Math/UVMetrics.h
@@ -0,0 +1,18 @@
+// Copyright Epic Games, Inc. All Rights Reserved.
+
+#pragma once
+#include "CoreMinimal.h"
+
+namespace UE::Geometry
+{
+	class FDynamicMesh3;
+
+	class UVEDITORTOOLS_API FUVMetrics
+	{
+	public:
+		static double ReedBeta(const UE::Geometry::FDynamicMesh3& Mesh, int32 UVChannel, int32 Tid);
+		static double Sander(const UE::Geometry::FDynamicMesh3& Mesh, int32 UVChannel, int32 Tid, bool bUseL2);
+		static double TexelDensity(const UE::Geometry::FDynamicMesh3& Mesh, int32 UVChannel, int32 Tid, int32 MapSize);
+	};
+
+}
\ No newline at end of file
diff --git a/Engine/Plugins/Experimental/UVEditor/Source/UVEditorTools/Public/ToolTargets/UVEditorToolMeshInput.h b/Engine/Plugins/Experimental/UVEditor/Source/UVEditorTools/Public/ToolTargets/UVEditorToolMeshInput.h
--- a/Engine/Plugins/Experimental/UVEditor/Source/UVEditorTools/Public/ToolTargets/UVEditorToolMeshInput.h
+++ b/Engine/Plugins/Experimental/UVEditor/Source/UVEditorTools/Public/ToolTargets/UVEditorToolMeshInput.h
@@ -5,6 +5,7 @@
 #include "CoreMinimal.h"
 
 #include "DynamicMesh/DynamicMeshAttributeSet.h" // FDynamicMeshUVOverlay
+#include "Components/BaseDynamicMeshComponent.h" // EMeshRenderAttributeFlags
 #include "ToolTargets/ToolTarget.h"
 #include "VectorTypes.h"
 
@@ -80,6 +81,12 @@ public:
 	UPROPERTY()
 	TObjectPtr<UMeshElementsVisualizer> WireframeDisplay = nullptr;
 
+	/**
+	* When enabled, this causes the unwrap preview to recieve updates to the vertex color buffers when mesh updates occur.
+	*/
+	UPROPERTY()
+	bool bEnableTriangleVertexColors = false;
+
 	// Information about an OnCanonicalModified broadcast. This is a struct so that we can
 	// add more information if we someday need to.
 	struct FCanonicalModifiedInfo
@@ -261,4 +268,16 @@ public:
 
 	// UToolTarget
 	virtual bool IsValid() const override;
+
+protected:	
+
+	/** 
+	* Returns the currently configured mesh attribute buffer flags enabled on this target for updating the unwrap previews
+	*/
+	virtual EMeshRenderAttributeFlags GetUnwrapUpdateAttributeFlags() const;
+
+	/**
+	* Returns the currently configured mesh attribute buffer flags enabled on this target for updating the applied previews
+	*/
+	virtual EMeshRenderAttributeFlags GetAppliedUpdateAttributeFlags() const;
 };
\ No newline at end of file
diff --git a/Engine/Plugins/Experimental/UVEditor/Source/UVEditorTools/Public/UVEditorUXSettings.h b/Engine/Plugins/Experimental/UVEditor/Source/UVEditorTools/Public/UVEditorUXSettings.h
--- a/Engine/Plugins/Experimental/UVEditor/Source/UVEditorTools/Public/UVEditorUXSettings.h
+++ b/Engine/Plugins/Experimental/UVEditor/Source/UVEditorTools/Public/UVEditorUXSettings.h
@@ -42,6 +42,14 @@ public:
 	static FLinearColor GetWireframeColorByTargetIndex(int32 TargetIndex);
 	static FLinearColor GetBoundaryColorByTargetIndex(int32 TargetIndex);
 
+	// Provides a color blindness friendly ramp for visualizing metrics and other data
+	// Domain is [0 , 1]
+	static FColor MakeCividisColorFromScalar(float Scalar);
+
+	// Provides a color blindness friendly ramp for visualizing diverging metrics and other data
+	// Domain is [-0.5, 0.5] 
+	static FColor MakeTurboColorFromScalar(float Scalar);
+
 	// Selection Highlighting Properties
 	static const float SelectionTriangleOpacity;
 	static const FColor SelectionTriangleFillColor;
