UVEditor: Fix Non-Unity compile build error

#rnx

diff --git a/Engine/Plugins/Experimental/UVEditor/Source/UVEditor/Private/UVEditorMode.cpp b/Engine/Plugins/Experimental/UVEditor/Source/UVEditor/Private/UVEditorMode.cpp
--- a/Engine/Plugins/Experimental/UVEditor/Source/UVEditor/Private/UVEditorMode.cpp
+++ b/Engine/Plugins/Experimental/UVEditor/Source/UVEditor/Private/UVEditorMode.cpp
@@ -52,6 +52,7 @@
 #include "Parameterization/MeshUDIMClassifier.h"
 #include "UVEditorLogging.h"
 #include "UObject/ObjectSaveContext.h"
+#include "Materials/Material.h"
 
 #include UE_INLINE_GENERATED_CPP_BY_NAME(UVEditorMode)
 
