Fix outdated icon path for UVEditor after move out of experimental

diff --git a/Engine/Plugins/Editor/UVEditor/Source/UVEditorTools/Private/UVEditorStyle.cpp b/Engine/Plugins/Editor/UVEditor/Source/UVEditorTools/Private/UVEditorStyle.cpp
--- a/Engine/Plugins/Editor/UVEditor/Source/UVEditorTools/Private/UVEditorStyle.cpp
+++ b/Engine/Plugins/Editor/UVEditor/Source/UVEditorTools/Private/UVEditorStyle.cpp
@@ -18,7 +18,7 @@ FUVEditorStyle::FUVEditorStyle()
 	const FVector2D IconSize(16.0f, 16.0f);
 	const FVector2D ToolbarIconSize(20.0f, 20.0f);
 
-	SetContentRoot(FPaths::EnginePluginsDir() / TEXT("Experimental/UVEditor/Content/Icons"));
+	SetContentRoot(FPaths::EnginePluginsDir() / TEXT("Editor/UVEditor/Content/Icons"));
 	SetCoreContentRoot(FPaths::EngineContentDir() / TEXT("Editor/Slate") );
 
 	Set("UVEditor.OpenUVEditor",				new IMAGE_BRUSH_SVG("UVEditor", IconSize));
