UVEditorModule: Prevent launching the UV Editor from context menu when read-only assets are selected in the content browser.

The "UV Editor" menu entry is now conditionally executable depending on whether or not the selected assets can be modified. When it is disabled, the hover tooltip is overridden to read "The selected asset(s) are read-only and cannot be edited."

diff --git a/Engine/Plugins/Editor/UVEditor/Source/UVEditor/Private/UVEditorModule.cpp b/Engine/Plugins/Editor/UVEditor/Source/UVEditor/Private/UVEditorModule.cpp
--- a/Engine/Plugins/Editor/UVEditor/Source/UVEditor/Private/UVEditorModule.cpp
+++ b/Engine/Plugins/Editor/UVEditor/Source/UVEditor/Private/UVEditorModule.cpp
@@ -107,9 +107,11 @@ void FUVEditorModule::RegisterMenus()
 								{
 									UVSubsystem->StartUVEditor(AssetsToEdit);
 								}
-							}));
+							}),
+							FCanExecuteAction::CreateWeakLambda(Context, [Context]() { return Context->bCanBeModified; }));
 
-						Section.AddMenuEntryWithCommandList(FUVEditorCommands::Get().OpenUVEditor, CommandListToBind, TAttribute<FText>(), TAttribute<FText>(), FSlateIcon(FUVEditorStyle::Get().GetStyleSetName(), "UVEditor.OpenUVEditor"));
+						const TAttribute<FText> ToolTipOverride = Context->bCanBeModified ? TAttribute<FText>() : LOCTEXT("ReadOnlyAssetWarning", "The selected asset(s) are read-only and cannot be edited.");
+						Section.AddMenuEntryWithCommandList(FUVEditorCommands::Get().OpenUVEditor, CommandListToBind, TAttribute<FText>(), ToolTipOverride, FSlateIcon(FUVEditorStyle::Get().GetStyleSetName(), "UVEditor.OpenUVEditor"));
 					}
 				}
 			}));
