UV Editor: Move gizmo menu entries to the bottom of the display menu, to help mitigate the current bug where dragging the mouse over the menu causes incorrect menu closure when interacting with other menu elements.

diff --git a/Engine/Plugins/Editor/UVEditor/Source/UVEditor/Private/UVEditorToolkit.cpp b/Engine/Plugins/Editor/UVEditor/Source/UVEditor/Private/UVEditorToolkit.cpp
--- a/Engine/Plugins/Editor/UVEditor/Source/UVEditor/Private/UVEditorToolkit.cpp
+++ b/Engine/Plugins/Editor/UVEditor/Source/UVEditor/Private/UVEditorToolkit.cpp
@@ -521,29 +521,6 @@ void FUVEditorToolkit::PostInitAssetEditor()
 			
 			TSharedRef<SVerticalBox> Container = SNew(SVerticalBox);
 
-			FMenuBuilder MenuBuilder(true, TSharedPtr<FUICommandList>());
-			MenuBuilder.BeginSection("Section_Gizmo", LOCTEXT("Section_Gizmo", "Gizmo"));
-
-			MenuBuilder.AddSubMenu(
-				LOCTEXT("GizmoTransformPanelSubMenu", "Gizmo Transform Panel"), LOCTEXT("GizmoTransformPanelSubMenu_ToolTip", "Configure the Gizmo Transform Panel."),
-				FNewMenuDelegate::CreateLambda([UVModeToolkit](FMenuBuilder& SubMenuBuilder) {
-					UVModeToolkit->MakeGizmoNumericalUISubMenu(SubMenuBuilder);
-				}));
-
-			MenuBuilder.EndSection();
-
-
-			Container->AddSlot()
-			.AutoHeight()
-			.Padding(FMargin(0.f, 0.f, 8.f, 0.f))
-			[
-				SNew(SBox)
-				.MinDesiredWidth(500)
-				[
-					MenuBuilder.MakeWidget()
-				]
-			];
-
 			Container->AddSlot()
 				.AutoHeight()
 				.Padding(FMargin(0.f, 0.f, 8.f, 0.f))
@@ -603,6 +580,33 @@ void FUVEditorToolkit::PostInitAssetEditor()
 				]
 				];
 
+			// TODO: We're moving this to the bottom because there's currently a bug with menu handling where mousing
+			// over this entry will break other menus in the overall display menu. 
+			// The bug causes premature menu closure, preventing settings from being applied to properties.
+
+			FMenuBuilder MenuBuilder(true, TSharedPtr<FUICommandList>(), TSharedPtr<FExtender>(), true);
+			MenuBuilder.BeginSection("Section_Gizmo", LOCTEXT("Section_Gizmo", "Gizmo"));
+
+			MenuBuilder.AddSubMenu(
+				LOCTEXT("GizmoTransformPanelSubMenu", "Gizmo Transform Panel"), LOCTEXT("GizmoTransformPanelSubMenu_ToolTip", "Configure the Gizmo Transform Panel."),
+				FNewMenuDelegate::CreateLambda([UVModeToolkit](FMenuBuilder& SubMenuBuilder) {
+					UVModeToolkit->MakeGizmoNumericalUISubMenu(SubMenuBuilder);
+				}));
+
+			MenuBuilder.EndSection();
+
+
+			Container->AddSlot()
+			.AutoHeight()
+			.Padding(FMargin(0.f, 0.f, 8.f, 0.f))
+			[
+				SNew(SBox)
+				.MinDesiredWidth(500)
+				[
+					MenuBuilder.MakeWidget()
+				]
+			];
+
 			TSharedRef<SWidget> Widget = SNew(SBorder)
 				.HAlign(HAlign_Fill)
 				.Padding(4)
