UVEditor: Adds support for a new tool to manipluate Texel Density for UV Islands.

Has the following features:

* Specify texel density as a ratio of two numbers: Distance in worldspace and number of pixels.
* Three modes - Scale Islands, Scale Whole, and Normalize.
* Is UDIM aware.
* Texel Density sampling mode - interactively select a triangle in the viewports to automatically setup the texel density ratios based on that triangle, useful in multi-object consistency workflows.

diff --git a/Engine/Plugins/Editor/UVEditor/Source/UVEditor/Private/UVEditorCommands.cpp b/Engine/Plugins/Editor/UVEditor/Source/UVEditor/Private/UVEditorCommands.cpp
--- a/Engine/Plugins/Editor/UVEditor/Source/UVEditor/Private/UVEditorCommands.cpp
+++ b/Engine/Plugins/Editor/UVEditor/Source/UVEditor/Private/UVEditorCommands.cpp
@@ -28,6 +28,7 @@ void FUVEditorCommands::RegisterCommands()
 	UI_COMMAND(BeginTransformTool, "Transform", "Transform existing UVs", EUserInterfaceActionType::ToggleButton, FInputChord());
 	UI_COMMAND(BeginAlignTool, "Align", "Align existing UVs", EUserInterfaceActionType::ToggleButton, FInputChord());
 	UI_COMMAND(BeginDistributeTool, "Distribute", "Distribute existing UVs", EUserInterfaceActionType::ToggleButton, FInputChord());
+	UI_COMMAND(BeginTexelDensityTool, "Texel Density", "Modify UVs based on texel density", EUserInterfaceActionType::ToggleButton, FInputChord());
 	UI_COMMAND(BeginParameterizeMeshTool, "AutoUV", "Auto-unwrap and pack UVs", EUserInterfaceActionType::ToggleButton, FInputChord());
 	UI_COMMAND(BeginChannelEditTool, "Channels", "Modify UV channels", EUserInterfaceActionType::ToggleButton, FInputChord());
 	UI_COMMAND(BeginSeamTool, "Seam", "Edit UV seams", EUserInterfaceActionType::ToggleButton, FInputChord());
diff --git a/Engine/Plugins/Editor/UVEditor/Source/UVEditor/Private/UVEditorMode.cpp b/Engine/Plugins/Editor/UVEditor/Source/UVEditor/Private/UVEditorMode.cpp
--- a/Engine/Plugins/Editor/UVEditor/Source/UVEditor/Private/UVEditorMode.cpp
+++ b/Engine/Plugins/Editor/UVEditor/Source/UVEditor/Private/UVEditorMode.cpp
@@ -39,6 +39,7 @@
 #include "UVEditorSeamTool.h"
 #include "UVEditorRecomputeUVsTool.h"
 #include "UVSelectTool.h"
+#include "UVEditorTexelDensityTool.h"
 #include "UVEditorInitializationContext.h"
 #include "UVEditorModeToolkit.h"
 #include "UVEditorSubsystem.h"
@@ -520,6 +521,10 @@ void UUVEditorMode::RegisterTools()
 	UVEditorDistributeToolBuilder->Targets = &ToolInputObjects;
 	RegisterTool(CommandInfos.BeginDistributeTool, TEXT("BeginDistributeTool"), UVEditorDistributeToolBuilder);
 
+	UUVEditorTexelDensityToolBuilder* UVEditorTexelDensityToolBuilder = NewObject<UUVEditorTexelDensityToolBuilder>();
+	UVEditorTexelDensityToolBuilder->Targets = &ToolInputObjects;
+	RegisterTool(CommandInfos.BeginTexelDensityTool, TEXT("BeginTexelDensityTool"), UVEditorTexelDensityToolBuilder);
+
 	UUVEditorParameterizeMeshToolBuilder* UVEditorParameterizeMeshToolBuilder = NewObject<UUVEditorParameterizeMeshToolBuilder>();
 	UVEditorParameterizeMeshToolBuilder->Targets = &ToolInputObjects;
 	RegisterTool(CommandInfos.BeginParameterizeMeshTool, TEXT("BeginParameterizeMeshTool"), UVEditorParameterizeMeshToolBuilder);
diff --git a/Engine/Plugins/Editor/UVEditor/Source/UVEditor/Private/UVEditorModeToolkit.cpp b/Engine/Plugins/Editor/UVEditor/Source/UVEditor/Private/UVEditorModeToolkit.cpp
--- a/Engine/Plugins/Editor/UVEditor/Source/UVEditor/Private/UVEditorModeToolkit.cpp
+++ b/Engine/Plugins/Editor/UVEditor/Source/UVEditor/Private/UVEditorModeToolkit.cpp
@@ -464,6 +464,7 @@ void FUVEditorModeToolkit::BuildToolPalette(FName PaletteIndex, class FToolBarBu
 		ToolbarBuilder.AddToolBarButton(Commands.BeginTransformTool);
 		ToolbarBuilder.AddToolBarButton(Commands.BeginAlignTool);
 		ToolbarBuilder.AddToolBarButton(Commands.BeginDistributeTool);
+		ToolbarBuilder.AddToolBarButton(Commands.BeginTexelDensityTool);
 		ToolbarBuilder.AddToolBarButton(Commands.BeginChannelEditTool);
 		ToolbarBuilder.AddToolBarButton(Commands.BeginSeamTool);
 		ToolbarBuilder.AddToolBarButton(Commands.BeginParameterizeMeshTool);
diff --git a/Engine/Plugins/Editor/UVEditor/Source/UVEditor/Public/UVEditorCommands.h b/Engine/Plugins/Editor/UVEditor/Source/UVEditor/Public/UVEditorCommands.h
--- a/Engine/Plugins/Editor/UVEditor/Source/UVEditor/Public/UVEditorCommands.h
+++ b/Engine/Plugins/Editor/UVEditor/Source/UVEditor/Public/UVEditorCommands.h
@@ -22,6 +22,7 @@ public:
 	TSharedPtr<FUICommandInfo> BeginTransformTool;
 	TSharedPtr<FUICommandInfo> BeginAlignTool;
 	TSharedPtr<FUICommandInfo> BeginDistributeTool;
+	TSharedPtr<FUICommandInfo> BeginTexelDensityTool;
 	TSharedPtr<FUICommandInfo> BeginParameterizeMeshTool;
 	TSharedPtr<FUICommandInfo> BeginChannelEditTool;
 	TSharedPtr<FUICommandInfo> BeginSeamTool;
diff --git a/Engine/Plugins/Editor/UVEditor/Source/UVEditorTools/Private/Operators/UVEditorTexelDensityOp.cpp b/Engine/Plugins/Editor/UVEditor/Source/UVEditorTools/Private/Operators/UVEditorTexelDensityOp.cpp
new file mode 100644
--- /dev/null
+++ b/Engine/Plugins/Editor/UVEditor/Source/UVEditorTools/Private/Operators/UVEditorTexelDensityOp.cpp
@@ -0,0 +1,360 @@
+// Copyright Epic Games, Inc. All Rights Reserved.
+
+#include "Operators/UVEditorTexelDensityOp.h"
+
+#include "DynamicMesh/DynamicMesh3.h"
+#include "DynamicMesh/DynamicMeshAttributeSet.h"
+#include "Parameterization/MeshUVPacking.h"
+#include "Properties/UVLayoutProperties.h"
+#include "Selections/MeshConnectedComponents.h"
+#include "Parameterization/MeshUDIMClassifier.h"
+#include "UDIMUtilities.h"
+#include "Math/UVMetrics.h"
+
+
+#include UE_INLINE_GENERATED_CPP_BY_NAME(UVEditorTexelDensityOp)
+
+using namespace UE::Geometry;
+
+namespace UVTexelDensityOpLocals
+{
+	FVector2f ExternalUVToInternalUV(const FVector2f& UV)
+	{
+		return FVector2f(UV.X, 1 - UV.Y);
+	}
+
+	FVector2f InternalUVToExternalUV(const FVector2f& UV)
+	{
+		return FVector2f(UV.X, 1 - UV.Y);
+	}
+
+	struct TileConnectedComponents
+	{
+		TileConnectedComponents(const FMeshConnectedComponents& ConnectedComponentsIn, const FVector2i& TileIn)
+			: ConnectedComponents(ConnectedComponentsIn),
+		Tile(TileIn)
+		{
+			for (int32 ComponentIndex = 0; ComponentIndex < ConnectedComponents.Num(); ++ComponentIndex)
+			{
+				TileTids.Append(ConnectedComponents[ComponentIndex].Indices);
+			}
+		}
+
+		TArray<int32> TileTids;
+		FMeshConnectedComponents ConnectedComponents;
+		FVector2i Tile;
+	};
+
+	void CollectIslandComponentsPerTile(const FDynamicMesh3& Mesh, const FDynamicMeshUVOverlay& UVOverlay, TOptional<TSet<int32>>& Selection,
+		                                TArray< TileConnectedComponents >& ComponentsPerTile, bool bUDIMsEnabled)
+	{
+		ComponentsPerTile.Empty();
+
+		TArray<FVector2i> Tiles;
+		TArray<TUniquePtr<TArray<int32>>> TileTids;
+
+		if (bUDIMsEnabled)
+		{
+			TOptional<TArray<int32>> SelectionArray;
+			if (Selection.IsSet())
+			{
+				SelectionArray = Selection.GetValue().Array();
+			}
+			FDynamicMeshUDIMClassifier TileClassifier(&UVOverlay, SelectionArray);
+
+			Tiles = TileClassifier.ActiveTiles();
+			for (const FVector2i& Tile : Tiles)
+			{
+				TileTids.Emplace(MakeUnique<TArray<int32>>(TileClassifier.TidsForTile(Tile)));
+			}
+		}
+		else
+		{
+			if (Selection.IsSet())
+			{
+				Tiles.Add({ 0,0 });
+				TileTids.Emplace(MakeUnique<TArray<int32>>(Selection.GetValue().Array()));
+			}
+			else
+			{
+				Tiles.Add({ 0,0 });
+				TileTids.Emplace(MakeUnique<TArray<int32>>());
+				TileTids[0]->Reserve(Mesh.TriangleCount());
+				for (int32 Tid : Mesh.TriangleIndicesItr())
+				{
+					TileTids[0]->Add(Tid);
+				}
+			}
+		}
+
+		for (int32 TileIndex = 0; TileIndex < Tiles.Num(); ++TileIndex)
+		{
+			FMeshConnectedComponents ConnectedComponents(&Mesh);
+			if (TileTids[TileIndex])
+			{
+				ConnectedComponents.FindConnectedTriangles(*TileTids[TileIndex], [&](int32 Triangle0, int32 Triangle1) {
+					return UVOverlay.AreTrianglesConnected(Triangle0, Triangle1);
+				});
+			}
+			else
+			{
+				ConnectedComponents.FindConnectedTriangles([&](int32 Triangle0, int32 Triangle1) {
+					return UVOverlay.AreTrianglesConnected(Triangle0, Triangle1);
+				});
+			}
+
+			if (!ConnectedComponents.Components.IsEmpty())
+			{
+				ComponentsPerTile.Emplace(ConnectedComponents, Tiles[TileIndex]);
+			}
+		}
+	}
+}
+
+bool UUVEditorTexelDensitySettings::InSamplingMode() const
+{
+	return true;
+}
+
+
+void FUVEditorTexelDensityOp::SetTransform(const FTransformSRT3d& Transform)
+{
+	ResultTransform = Transform;
+}
+
+
+void FUVEditorTexelDensityOp::ScaleMeshSubRegionByDensity(FDynamicMeshUVOverlay* UVLayer, const TArray<int32>& Tids, TSet<int32>& UVElements, int32 TileResolution)
+{
+	double AverageTidTexelDensity = 0.0;
+	int SetTriangleCount = 0;
+	for (int Tid : Tids)
+	{
+		if (ResultMesh->Attributes()->GetUVLayer(UVLayerIndex)->IsSetTriangle(Tid))
+		{
+			AverageTidTexelDensity += FUVMetrics::TexelDensity(*ResultMesh, UVLayerIndex, Tid, TileResolution);
+			SetTriangleCount++;
+		}
+	}
+	if (SetTriangleCount)
+	{
+		AverageTidTexelDensity /= SetTriangleCount;
+	}
+
+	double TargetTexelDensity = (double)TargetPixelCountMeasurement / TargetWorldSpaceMeasurement;
+	double RequiredGlobalScaleFactor = TargetTexelDensity / AverageTidTexelDensity;
+
+	for (int ElementID : UVElements)
+	{
+		FVector2f UV = UVTexelDensityOpLocals::InternalUVToExternalUV(UVLayer->GetElement(ElementID));
+		UV = (UV * RequiredGlobalScaleFactor);
+		UVLayer->SetElement(ElementID, UVTexelDensityOpLocals::ExternalUVToInternalUV(UV));
+	}
+}
+
+
+
+void FUVEditorTexelDensityOp::CalculateResult(FProgressCancel* Progress)
+{
+	TRACE_CPUPROFILER_EVENT_SCOPE(UVLayoutOp_CalculateResult);
+
+	if (Progress && Progress->Cancelled())
+	{
+		return;
+	}
+	ResultMesh->Copy(*OriginalMesh, true, true, true, true);
+
+	if (!ensureMsgf(ResultMesh->HasAttributes(), TEXT("Attributes not found on mesh? Conversion should always create them, so this operator should not need to do so.")))
+	{
+		ResultMesh->EnableAttributes();
+	}
+
+	if (Progress && Progress->Cancelled())
+	{
+		return;
+	}
+
+
+	int UVLayerInput = UVLayerIndex;
+	FDynamicMeshUVOverlay* UseUVLayer = ResultMesh->Attributes()->GetUVLayer(UVLayerInput);
+
+
+	if (Progress && Progress->Cancelled())
+	{
+		return;
+	}
+
+	if (bMaintainOriginatingUDIM)
+	{
+		TArray<UVTexelDensityOpLocals::TileConnectedComponents> TileComponents;
+		UVTexelDensityOpLocals::CollectIslandComponentsPerTile(*ResultMesh, *ResultMesh->Attributes()->GetUVLayer(UVLayerIndex), Selection, TileComponents, true);
+
+		for (UVTexelDensityOpLocals::TileConnectedComponents Tile : TileComponents)
+		{
+			FVector2i TileIndex = Tile.Tile;	
+			const int32 TileID = UE::TextureUtilitiesCommon::GetUDIMIndex(TileIndex.X, TileIndex.Y);
+			float TileResolution = TextureResolution;
+			if (TextureResolutionPerUDIM.IsSet())
+			{
+				TileResolution = TextureResolutionPerUDIM.GetValue().FindOrAdd(TileID, this->TextureResolution);
+			}
+
+			TSet<int32> ElementsToMove;
+			ElementsToMove.Reserve(Tile.TileTids.Num() * 3);
+			for (int Tid : Tile.TileTids)
+			{
+				FIndex3i Elements = UseUVLayer->GetTriangle(Tid);
+				ElementsToMove.Add(Elements[0]);
+				ElementsToMove.Add(Elements[1]);
+				ElementsToMove.Add(Elements[2]);
+			}
+
+			for (int32 Element : ElementsToMove)
+			{
+				FVector2f UV = UVTexelDensityOpLocals::InternalUVToExternalUV(UseUVLayer->GetElement(Element));
+				UV = UV - FVector2f(TileIndex);
+				UseUVLayer->SetElement(Element, UVTexelDensityOpLocals::ExternalUVToInternalUV(UV));
+			}
+
+			// TODO: There is a second connected components call inside the packer that might be unnessessary. Could be a future optimization.
+			FDynamicMeshUVPacker Packer(UseUVLayer, MakeUnique<TArray<int32>>(Tile.TileTids));
+			Packer.TextureResolution = TileResolution;
+			ExecutePacker(Packer);
+			if (Progress && Progress->Cancelled())
+			{
+				return;
+			}
+
+			if (TexelDensityMode == EUVTexelDensityOpModes::ScaleGlobal)
+			{
+				ScaleMeshSubRegionByDensity(UseUVLayer, Tile.TileTids, ElementsToMove, TileResolution);
+			}
+
+			if (TexelDensityMode == EUVTexelDensityOpModes::ScaleIslands)
+			{
+				for (int32 ComponentIndex = 0; ComponentIndex < Tile.ConnectedComponents.Components.Num(); ComponentIndex++)
+				{
+					const TArray<int32>& ComponentTids = Tile.ConnectedComponents.Components[ComponentIndex].Indices;
+
+					TSet<int32> ComponentElementsToMove;
+					ComponentElementsToMove.Reserve(ComponentTids.Num() * 3);
+					for (int Tid : ComponentTids)
+					{
+						FIndex3i Elements = UseUVLayer->GetTriangle(Tid);
+						ComponentElementsToMove.Add(Elements[0]);
+						ComponentElementsToMove.Add(Elements[1]);
+						ComponentElementsToMove.Add(Elements[2]);
+					}
+
+					ScaleMeshSubRegionByDensity(UseUVLayer, ComponentTids, ComponentElementsToMove, TileResolution);
+				}
+			}
+
+			for (int32 Element : ElementsToMove)
+			{
+				FVector2f UV = UVTexelDensityOpLocals::InternalUVToExternalUV(UseUVLayer->GetElement(Element));
+				UV = UV + FVector2f(TileIndex);
+				UseUVLayer->SetElement(Element, UVTexelDensityOpLocals::ExternalUVToInternalUV(UV));
+			}
+		}
+	}
+	else
+	{
+		TArray<UVTexelDensityOpLocals::TileConnectedComponents> TileComponents;
+		UVTexelDensityOpLocals::CollectIslandComponentsPerTile(*ResultMesh, *ResultMesh->Attributes()->GetUVLayer(UVLayerIndex), Selection, TileComponents, false);
+
+		TSet<int32> ElementsToMove;
+		ElementsToMove.Reserve(TileComponents[0].TileTids.Num() * 3);
+		for (int Tid : TileComponents[0].TileTids)
+		{
+			FIndex3i Elements = UseUVLayer->GetTriangle(Tid);
+			ElementsToMove.Add(Elements[0]);
+			ElementsToMove.Add(Elements[1]);
+			ElementsToMove.Add(Elements[2]);
+		}
+
+		FDynamicMeshUVPacker Packer(UseUVLayer, MakeUnique<TArray<int32>>(TileComponents[0].TileTids));
+		Packer.TextureResolution = TextureResolution;
+		ExecutePacker(Packer);
+		if (Progress && Progress->Cancelled())
+		{
+			return;
+		}
+
+		if (TexelDensityMode == EUVTexelDensityOpModes::ScaleGlobal)
+		{
+			ScaleMeshSubRegionByDensity(UseUVLayer, TileComponents[0].TileTids, ElementsToMove, TextureResolution);
+		}
+
+		if (TexelDensityMode == EUVTexelDensityOpModes::ScaleIslands)
+		{
+			for (int32 ComponentIndex = 0; ComponentIndex < TileComponents[0].ConnectedComponents.Components.Num(); ComponentIndex++)
+			{
+				const TArray<int32>& ComponentTids = TileComponents[0].ConnectedComponents.Components[ComponentIndex].Indices;
+
+				TSet<int32> ComponentElementsToMove;
+				ComponentElementsToMove.Reserve(ComponentTids.Num() * 3);
+				for (int Tid : ComponentTids)
+				{
+					FIndex3i Elements = UseUVLayer->GetTriangle(Tid);
+					ComponentElementsToMove.Add(Elements[0]);
+					ComponentElementsToMove.Add(Elements[1]);
+					ComponentElementsToMove.Add(Elements[2]);
+				}
+
+				ScaleMeshSubRegionByDensity(UseUVLayer, ComponentTids, ComponentElementsToMove, TextureResolution);
+			}
+		}
+
+	}
+
+
+}
+
+void FUVEditorTexelDensityOp::ExecutePacker(FDynamicMeshUVPacker& Packer)
+{
+	if (TexelDensityMode == EUVTexelDensityOpModes::Normalize)
+	{
+		Packer.bScaleIslandsByWorldSpaceTexelRatio = true;
+		if (Packer.StandardPack() == false)
+		{
+			// failed... what to do?
+			return;
+		}
+	}
+}
+
+TUniquePtr<FDynamicMeshOperator> UUVTexelDensityOperatorFactory::MakeNewOperator()
+{
+	TUniquePtr<FUVEditorTexelDensityOp> Op = MakeUnique<FUVEditorTexelDensityOp>();
+
+	Op->OriginalMesh = OriginalMesh;
+
+	switch (Settings->TexelDensityMode)
+	{
+	case ETexelDensityToolMode::ApplyToIslands:
+		Op->TexelDensityMode = EUVTexelDensityOpModes::ScaleIslands;
+		break;
+	case ETexelDensityToolMode::ApplyToWhole:
+		Op->TexelDensityMode = EUVTexelDensityOpModes::ScaleGlobal;
+		break;
+	case ETexelDensityToolMode::Normalize:
+		Op->TexelDensityMode = EUVTexelDensityOpModes::Normalize;
+		break;
+	default:
+		ensure(false);
+		break;
+	}
+
+	Op->TextureResolution = Settings->TextureResolution;
+	Op->TargetWorldSpaceMeasurement = Settings->TargetWorldUnits;
+	Op->TargetPixelCountMeasurement = Settings->TargetPixelCount;
+
+	Op->UVLayerIndex = GetSelectedUVChannel();
+	Op->TextureResolution = Settings->TextureResolution;
+	Op->SetTransform(TargetTransform);
+	Op->bMaintainOriginatingUDIM = Settings->bEnableUDIMLayout;
+	Op->Selection = Selection;
+	Op->TextureResolutionPerUDIM = TextureResolutionPerUDIM;
+
+	return Op;
+}
diff --git a/Engine/Plugins/Editor/UVEditor/Source/UVEditorTools/Private/UVEditorTexelDensityTool.cpp b/Engine/Plugins/Editor/UVEditor/Source/UVEditorTools/Private/UVEditorTexelDensityTool.cpp
new file mode 100644
--- /dev/null
+++ b/Engine/Plugins/Editor/UVEditor/Source/UVEditorTools/Private/UVEditorTexelDensityTool.cpp
@@ -0,0 +1,738 @@
+// Copyright Epic Games, Inc. All Rights Reserved.
+
+#include "UVEditorTexelDensityTool.h"
+
+#include "ContextObjectStore.h"
+#include "MeshOpPreviewHelpers.h" // UMeshOpPreviewWithBackgroundCompute
+#include "Operators/UVEditorTexelDensityOp.h"
+#include "InputBehaviorSet.h"
+#include "InputRouter.h"
+#include "BaseBehaviors/SingleClickBehavior.h"
+#include "BaseBehaviors/MouseHoverBehavior.h"
+#include "ToolSceneQueriesUtil.h"
+#include "Drawing/PreviewGeometryActor.h"
+#include "ToolSetupUtil.h"
+#include "UVEditorUXSettings.h"
+#include "ContextObjects/UVToolContextObjects.h"
+#include "Math/UVMetrics.h"
+
+#include UE_INLINE_GENERATED_CPP_BY_NAME(UVEditorTexelDensityTool)
+
+using namespace UE::Geometry;
+
+#define LOCTEXT_NAMESPACE "UUVEditorTexelDensityTool"
+
+namespace UVEditorTexelDensityToolLocals
+{
+	const FString& HoverTriangleSetID(TEXT("HoverTriangleSet"));
+	const FString& HoverLineSetID(TEXT("HoverLineSet"));
+}
+
+bool UUVEditorTexelDensityToolBuilder::CanBuildTool(const FToolBuilderState& SceneState) const
+{
+	return Targets && Targets->Num() > 0;
+}
+
+UInteractiveTool* UUVEditorTexelDensityToolBuilder::BuildTool(const FToolBuilderState& SceneState) const
+{
+	UUVEditorTexelDensityTool* NewTool = NewObject<UUVEditorTexelDensityTool>(SceneState.ToolManager);
+	NewTool->SetTargets(*Targets);
+
+	return NewTool;
+}
+
+void UUVEditorTexelDensityActionSettings::PostAction(ETexelDensityToolAction Action)
+{
+	ParentTool->RequestAction(Action);
+}
+
+
+void UUVEditorTexelDensityActionSettings::SampleTexelDensity()
+{
+	PostAction(ETexelDensityToolAction::BeginSamping);
+}
+
+
+bool UUVEditorTexelDensityToolSettings::InSamplingMode() const
+{
+	if (ParentTool.IsValid())
+	{
+		return ParentTool->ActiveAction() == ETexelDensityToolAction::Sampling || ParentTool->ActiveAction() == ETexelDensityToolAction::BeginSamping;
+	}
+	return true;
+}
+
+void UUVEditorTexelDensityTool::Setup()
+{
+	check(Targets.Num() > 0);
+
+	ToolStartTimeAnalytics = FDateTime::UtcNow();
+
+	UInteractiveTool::Setup();
+	UContextObjectStore* ContextStore = GetToolManager()->GetContextObjectStore();
+
+	Settings = NewObject<UUVEditorTexelDensityToolSettings>(this);
+	ActionSettings = NewObject<UUVEditorTexelDensityActionSettings>(this);
+	ActionSettings->Initialize(this);
+	Cast<UUVEditorTexelDensityToolSettings>(Settings)->Initialize(this);
+	Settings->RestoreProperties(this);
+	AddToolPropertySource(Settings);
+	AddToolPropertySource(ActionSettings);
+
+	UVToolSelectionAPI = ContextStore->FindContext<UUVToolSelectionAPI>();
+	UVTool2DViewportAPI = ContextStore->FindContext<UUVTool2DViewportAPI>();
+	LivePreviewAPI = ContextStore->FindContext<UUVToolLivePreviewAPI>();
+
+	UUVToolSelectionAPI::FHighlightOptions HighlightOptions;
+	HighlightOptions.bBaseHighlightOnPreviews = true;
+	HighlightOptions.bAutoUpdateUnwrap = true;
+	UVToolSelectionAPI->SetHighlightOptions(HighlightOptions);
+	UVToolSelectionAPI->SetHighlightVisible(true, false, true);
+
+	auto SetupOpFactory = [this](UUVEditorToolMeshInput& Target, const FUVToolSelection* Selection)
+	{
+		int32 TargetIndex = Targets.Find(&Target);
+
+		TObjectPtr<UUVTexelDensityOperatorFactory> Factory = NewObject<UUVTexelDensityOperatorFactory>();
+		Factory->TargetTransform = Target.AppliedPreview->PreviewMesh->GetTransform();
+		Factory->Settings = Settings;
+		Factory->OriginalMesh = Target.AppliedCanonical;
+		Factory->GetSelectedUVChannel = [&Target]() { return Target.UVLayerIndex; };
+		if (Selection)
+		{
+			Factory->Selection.Emplace(Selection->GetConvertedSelection(*Target.UnwrapCanonical, FUVToolSelection::EType::Triangle).SelectedIDs);
+		}
+
+		if (UVTool2DViewportAPI)
+		{
+			TMap<int32, int32> TextureResolutionPerUDIM;
+			for (const FUDIMBlock& UDIM : UVTool2DViewportAPI->GetUDIMBlocks())
+			{
+				TextureResolutionPerUDIM.Add(UDIM.UDIM, UDIM.TextureResolution);
+			}
+			Factory->TextureResolutionPerUDIM = TextureResolutionPerUDIM;
+		}
+
+		Target.AppliedPreview->ChangeOpFactory(Factory);
+		Target.AppliedPreview->OnMeshUpdated.AddWeakLambda(this, [this, &Target](UMeshOpPreviewWithBackgroundCompute* Preview) {
+			Target.UpdateUnwrapPreviewFromAppliedPreview();
+
+			this->UVToolSelectionAPI->RebuildUnwrapHighlight(Target.UnwrapPreview->PreviewMesh->GetTransform());
+			});
+
+		Target.AppliedPreview->OnOpCompleted.AddWeakLambda(this,
+			[this, TargetIndex](const FDynamicMeshOperator* Op)
+			{
+				RemainingTargetsToProcess--;
+			}
+		);
+
+		Target.AppliedPreview->InvalidateResult();
+		return Factory;
+	};
+
+	RemainingTargetsToProcess = 0;
+	if (UVToolSelectionAPI->HaveSelections())
+	{
+		Factories.Reserve(UVToolSelectionAPI->GetSelections().Num());
+		for (FUVToolSelection Selection : UVToolSelectionAPI->GetSelections())
+		{
+			RemainingTargetsToProcess++;
+			Factories.Add(SetupOpFactory(*Selection.Target, &Selection));
+		}
+	}
+	else
+	{
+		Factories.Reserve(Targets.Num());
+		for (int32 TargetIndex = 0; TargetIndex < Targets.Num(); ++TargetIndex)
+		{
+			RemainingTargetsToProcess++;
+			Factories.Add(SetupOpFactory(*Targets[TargetIndex], nullptr));
+		}
+	}
+
+	// Set things up for being able to add behaviors to live preview.
+	LivePreviewBehaviorSet = NewObject<UInputBehaviorSet>();
+	LivePreviewBehaviorSource = NewObject<ULocalInputBehaviorSource>();
+	LivePreviewBehaviorSource->GetInputBehaviorsFunc = [this]() { return LivePreviewBehaviorSet; };
+
+	// Set up click and hover behaviors both for the 2d (unwrap) viewport and the 3d 
+	// (applied/live preview) viewport
+	ULocalSingleClickInputBehavior* UnwrapClickBehavior = NewObject<ULocalSingleClickInputBehavior>();
+	UnwrapClickBehavior->Initialize();
+	UnwrapClickBehavior->IsHitByClickFunc = [this](const FInputDeviceRay& InputRay) {
+		FInputRayHit HitResult;
+		HitResult.bHit = Get2DHitTriangle(InputRay.WorldRay) != IndexConstants::InvalidID;
+		// We don't bother with the depth since there aren't competing click behaviors
+		return HitResult;
+	};
+	UnwrapClickBehavior->OnClickedFunc = [this](const FInputDeviceRay& ClickPos) {
+		int32 IndexOfMesh = IndexConstants::InvalidID;
+		int32 Tid = Get2DHitTriangle(ClickPos.WorldRay, &IndexOfMesh);
+		if (Tid != IndexConstants::InvalidID)
+		{
+			OnMeshTriangleClicked(Tid, IndexOfMesh, true);
+		}
+	};
+	AddInputBehavior(UnwrapClickBehavior);
+
+	ULocalMouseHoverBehavior* UnwrapHoverBehavior = NewObject<ULocalMouseHoverBehavior>();
+	UnwrapHoverBehavior->Initialize();
+	UnwrapHoverBehavior->BeginHitTestFunc = [this](const FInputDeviceRay& PressPos) {
+		FInputRayHit HitResult;
+		HitResult.bHit = Get2DHitTriangle(PressPos.WorldRay) != IndexConstants::InvalidID;
+		// We don't bother with the depth since there aren't competing hover behaviors
+		return HitResult;
+	};
+	UnwrapHoverBehavior->OnUpdateHoverFunc = [this](const FInputDeviceRay& DevicePos)
+	{
+		int32 IndexOfMesh = IndexConstants::InvalidID;
+		int32 Tid = Get2DHitTriangle(DevicePos.WorldRay, &IndexOfMesh);
+		if (Tid != IndexConstants::InvalidID)
+		{
+			OnMeshTriangleHovered(Tid, IndexOfMesh, true);
+			return true;
+		}
+		return false;
+	};
+	UnwrapHoverBehavior->OnEndHoverFunc = [this]()
+	{
+		ClearHover();
+	};
+	AddInputBehavior(UnwrapHoverBehavior);
+
+	// Now the applied/live preview viewport...
+	ULocalSingleClickInputBehavior* LivePreviewClickBehavior = NewObject<ULocalSingleClickInputBehavior>();
+	LivePreviewClickBehavior->Initialize();
+	LivePreviewClickBehavior->IsHitByClickFunc = [this](const FInputDeviceRay& InputRay) {
+		FInputRayHit HitResult;
+		HitResult.bHit = Get3DHitTriangle(InputRay.WorldRay) != IndexConstants::InvalidID;
+		// We don't bother with the depth since there aren't competing click behaviors
+		return HitResult;
+	};
+	LivePreviewClickBehavior->OnClickedFunc = [this](const FInputDeviceRay& ClickPos) {
+		int32 IndexOfMesh = IndexConstants::InvalidID;
+		int32 Tid = Get3DHitTriangle(ClickPos.WorldRay, &IndexOfMesh);
+		if (Tid != IndexConstants::InvalidID)
+		{
+			OnMeshTriangleClicked(Tid, IndexOfMesh, false);
+		}
+	};
+
+	LivePreviewBehaviorSet->Add(LivePreviewClickBehavior, this);
+
+	ULocalMouseHoverBehavior* LivePreviewHoverBehavior = NewObject<ULocalMouseHoverBehavior>();
+	LivePreviewHoverBehavior->Initialize();
+	LivePreviewHoverBehavior->BeginHitTestFunc = [this](const FInputDeviceRay& PressPos) {
+		FInputRayHit HitResult;
+		HitResult.bHit = Get3DHitTriangle(PressPos.WorldRay) != IndexConstants::InvalidID;
+		// We don't bother with the depth since there aren't competing hover behaviors
+		return HitResult;
+	};
+	LivePreviewHoverBehavior->OnUpdateHoverFunc = [this](const FInputDeviceRay& DevicePos)
+	{
+		int32 IndexOfMesh = IndexConstants::InvalidID;
+		int32 Tid = Get3DHitTriangle(DevicePos.WorldRay, &IndexOfMesh);
+		if (Tid != IndexConstants::InvalidID)
+		{
+			OnMeshTriangleHovered(Tid, IndexOfMesh, false);
+			return true;
+		}
+		return false;
+	};
+	LivePreviewHoverBehavior->OnEndHoverFunc = [this]()
+	{
+		ClearHover();
+	};
+	LivePreviewBehaviorSet->Add(LivePreviewHoverBehavior, this);
+
+	// Give the live preview behaviors to the live preview input router
+	if (LivePreviewAPI)
+	{
+		LivePreviewInputRouter = LivePreviewAPI->GetLivePreviewInputRouter();
+		LivePreviewInputRouter->RegisterSource(LivePreviewBehaviorSource);
+	}
+
+	// Set up all the components we need to visualize things.
+	UnwrapGeometry = NewObject<UPreviewGeometry>();
+	UnwrapGeometry->CreateInWorld(Targets[0]->UnwrapPreview->GetWorld(), FTransform::Identity);
+
+	LivePreviewGeometry = NewObject<UPreviewGeometry>();
+	LivePreviewGeometry->CreateInWorld(Targets[0]->AppliedPreview->GetWorld(), FTransform::Identity);
+
+	// These visualize the currently hovered (preview) portion of the seam
+	UnwrapGeometry->AddTriangleSet(UVEditorTexelDensityToolLocals::HoverTriangleSetID);
+	UnwrapGeometry->AddLineSet(UVEditorTexelDensityToolLocals::HoverLineSetID);
+	LivePreviewGeometry->AddTriangleSet(UVEditorTexelDensityToolLocals::HoverTriangleSetID);
+	LivePreviewGeometry->AddLineSet(UVEditorTexelDensityToolLocals::HoverLineSetID);
+
+	TriangleSetMaterial = ToolSetupUtil::GetCustomTwoSidedDepthOffsetMaterial(this->GetToolManager(),
+		FUVEditorUXSettings::SelectionHoverTriangleFillColor,
+		FUVEditorUXSettings::SelectionHoverTriangleDepthBias,
+		FUVEditorUXSettings::SelectionHoverTriangleOpacity);
+
+	UUVToolAABBTreeStorage* TreeStore = ContextStore->FindContext<UUVToolAABBTreeStorage>();
+	if (!TreeStore)
+	{
+		TreeStore = NewObject<UUVToolAABBTreeStorage>();
+		ContextStore->AddContextObject(TreeStore);
+	}
+
+	// Initialize the AABB trees from cached values, or make new ones.
+	for (int32 i = 0; i < Targets.Num(); ++i)
+	{
+		TSharedPtr<FDynamicMeshAABBTree3> Tree2d = TreeStore->Get(Targets[i]->UnwrapCanonical.Get());
+		if (!Tree2d)
+		{
+			Tree2d = MakeShared<FDynamicMeshAABBTree3>();
+			Tree2d->SetMesh(Targets[i]->UnwrapCanonical.Get());
+			TreeStore->Set(Targets[i]->UnwrapCanonical.Get(), Tree2d, Targets[i]);
+		}
+		Spatials2D.Add(Tree2d);
+
+		TSharedPtr<FDynamicMeshAABBTree3> Tree3d = TreeStore->Get(Targets[i]->AppliedCanonical.Get());
+		if (!Tree3d)
+		{
+			Tree3d = MakeShared<FDynamicMeshAABBTree3>();
+			Tree3d->SetMesh(Targets[i]->AppliedCanonical.Get());
+			TreeStore->Set(Targets[i]->AppliedCanonical.Get(), Tree3d, Targets[i]);
+		}
+		Spatials3D.Add(Tree3d);
+
+	}
+
+
+	SetToolDisplayName(LOCTEXT("ToolNameLocal", "Texel Density"));
+	GetToolManager()->DisplayMessage(
+		LOCTEXT("OnStartTool_TexelDensity", "Read and rescale UVs based on texel density values."),
+		EToolMessageLevel::UserNotification);
+
+	// Analytics
+	InputTargetAnalytics = UVEditorAnalytics::CollectTargetAnalytics(Targets);
+}
+
+int32 UUVEditorTexelDensityTool::Get2DHitTriangle(const FRay& WorldRayIn, int32* IndexOf2DSpatialOut)
+{
+	TRACE_CPUPROFILER_EVENT_SCOPE(UVEditorTexelDensityTool_Get2DHitTriangle);
+
+	if (!ensure(WorldRayIn.Direction.Z != 0))
+	{
+		return IndexConstants::InvalidID;
+	}
+
+	// Find plane hit location
+	double PlaneHitT = WorldRayIn.Origin.Z / -WorldRayIn.Direction.Z;
+	FVector3d PlanePoint = WorldRayIn.PointAt(PlaneHitT);
+
+	int32 Result = IndexConstants::InvalidID;
+	// For tolerance, we'd like to divide the screen into 90 sections to be similar to the 90 degree
+	// FOV angle tolerance we use in 3D. Ray Z should be our distance from ground plane, so width
+	// would be 2Z for 90 degree FOV.
+	double Tolerance = ToolSceneQueriesUtil::GetDefaultVisualAngleSnapThreshD() * WorldRayIn.Origin.Z / 45;
+
+	double MinDistSquared = TNumericLimits<double>::Max();
+	for (int32 i = 0; i < Spatials2D.Num(); ++i)
+	{
+		double DistSquared = TNumericLimits<double>::Max();
+		int32 Tid = Spatials2D[i]->FindNearestTriangle(PlanePoint, DistSquared);
+		if (DistSquared >= 0 && DistSquared < MinDistSquared)
+		{
+			MinDistSquared = DistSquared;
+			Result = Tid;
+			if (IndexOf2DSpatialOut)
+			{
+				*IndexOf2DSpatialOut = i;
+			}
+		}
+	}
+
+	return Result;
+}
+
+int32 UUVEditorTexelDensityTool::Get3DHitTriangle(const FRay& WorldRayIn, int32* IndexOf3DSpatialOut)
+{
+	TRACE_CPUPROFILER_EVENT_SCOPE(UVEditorTexelDensityTool_Get3DHitTriangle);
+	// We'll hit test all the preview meshes
+	double BestRayT = TNumericLimits<double>::Max();
+	int32 BestTid = IndexConstants::InvalidID;
+	for (int32 i = 0; i < Spatials3D.Num(); ++i)
+	{
+		FTransform Transform = Targets[i]->AppliedPreview->PreviewMesh->GetTransform();
+		FRay3d LocalRay(
+			(FVector3d)Transform.InverseTransformPosition(WorldRayIn.Origin),
+			(FVector3d)Transform.InverseTransformVector(WorldRayIn.Direction));
+
+		int32 Tid = IndexConstants::InvalidID;
+		double RayT = TNumericLimits<double>::Max();
+		if (Spatials3D[i]->FindNearestHitTriangle(LocalRay, RayT, Tid))
+		{
+			if (RayT >= 0 && RayT < BestRayT)
+			{
+				BestRayT = RayT;
+				BestTid = Tid;
+				if (IndexOf3DSpatialOut)
+				{
+					*IndexOf3DSpatialOut = i;
+				}
+			}
+		}
+	}
+
+	if (BestTid == IndexConstants::InvalidID)
+	{
+		// Didn't hit anything
+		return IndexConstants::InvalidID;
+	}
+
+	return BestTid;
+}
+
+void UUVEditorTexelDensityTool::OnMeshTriangleClicked(int32 Tid, int32 IndexOfMesh, bool bTidIsFromUnwrap)
+{
+	using namespace UVEditorTexelDensityToolLocals;
+
+	if (ClickedTid != IndexConstants::InvalidID)
+	{
+		// Haven't yet finished processing the last click
+		return;
+	}
+
+	// Save the click information. The click will be applied on tick.
+	ClickedTid = Tid;
+	ClickedMeshIndex = IndexOfMesh;
+	bClickWasInUnwrap = bTidIsFromUnwrap;
+}
+
+void UUVEditorTexelDensityTool::OnMeshTriangleHovered(int32 Tid, int32 IndexOfMesh, bool bTidIsFromUnwrap)
+{
+	using namespace UVEditorTexelDensityToolLocals;
+
+	// Save the hover information. Hover will be updated on tick
+	HoverTid = Tid;
+	HoverMeshIndex = IndexOfMesh;
+	bHoverTidIsFromUnwrap = bTidIsFromUnwrap;
+}
+
+void UUVEditorTexelDensityTool::UpdateHover()
+{
+	using namespace UVEditorTexelDensityToolLocals;
+
+	if (bLastHoverTidWasFromUnwrap == bHoverTidIsFromUnwrap
+		&& LastHoverMeshIndex == HoverMeshIndex
+		&& LastHoverTid == HoverTid)
+	{
+		return;
+	}
+
+	bLastHoverTidWasFromUnwrap = bHoverTidIsFromUnwrap;
+	LastHoverMeshIndex = HoverMeshIndex;
+	LastHoverTid = HoverTid;
+
+	ClearHover(false);
+
+	if (HoverTid == IndexConstants::InvalidID)
+	{
+		return;
+	}
+
+	UTriangleSetComponent* UnwrapTriangleSet = UnwrapGeometry->FindTriangleSet(HoverTriangleSetID);
+	UTriangleSetComponent* AppliedTriangleSet = LivePreviewGeometry->FindTriangleSet(HoverTriangleSetID);
+
+	FVector VertA, VertB, VertC;
+
+	Targets[HoverMeshIndex]->UnwrapCanonical->GetTriVertices(HoverTid, VertA, VertB, VertC);
+	UnwrapTriangleSet->AddTriangle(FRenderableTriangle(
+		TriangleSetMaterial,
+		{ VertA, FVector2D(), FVector(0,0,1), FColor() },
+		{ VertB, FVector2D(), FVector(0,0,1), FColor() },
+		{ VertC, FVector2D(), FVector(0,0,1), FColor() }
+	));
+
+	FTransform3d Transform = Targets[HoverMeshIndex]->AppliedPreview->PreviewMesh->GetTransform();
+
+	//UE_LOG(UVEditorLog, Log, TEXT("Mesh: %d; Transform: %f %f %f;"), HoverMeshIndex, Transform.GetTranslation().X, Transform.GetTranslation().Y, Transform.GetTranslation().Z);
+
+	Targets[HoverMeshIndex]->AppliedCanonical->GetTriVertices(HoverTid, VertA, VertB, VertC);
+
+	//UE_LOG(UVEditorLog, Log, TEXT("VECA: %f %f %f;"), VertA.X, VertA.Y, VertA.Z);
+
+	AppliedTriangleSet->AddTriangle(FRenderableTriangle(
+		TriangleSetMaterial,
+		{ Transform.TransformPosition(VertA), FVector2D(), FVector(0,0,1), FColor() },
+		{ Transform.TransformPosition(VertB), FVector2D(), FVector(0,0,1), FColor() },
+		{ Transform.TransformPosition(VertC), FVector2D(), FVector(0,0,1), FColor() }
+	));
+
+}
+
+void UUVEditorTexelDensityTool::ClearHover(bool bClearHoverInfo)
+{
+	using namespace UVEditorTexelDensityToolLocals;
+
+	UnwrapGeometry->FindTriangleSet(HoverTriangleSetID)->Clear();
+	LivePreviewGeometry->FindTriangleSet(HoverTriangleSetID)->Clear();
+	UnwrapGeometry->FindLineSet(HoverLineSetID)->Clear();
+	LivePreviewGeometry->FindLineSet(HoverLineSetID)->Clear();
+
+	if (bClearHoverInfo)
+	{
+		HoverTid = IndexConstants::InvalidID;
+		HoverMeshIndex = IndexConstants::InvalidID;
+		LastHoverTid = HoverTid;
+		LastHoverMeshIndex = HoverMeshIndex;
+	}
+}
+
+void UUVEditorTexelDensityTool::ApplyClick()
+{
+	const TObjectPtr<UUVEditorToolMeshInput>& Target = Targets[ClickedMeshIndex];
+
+	bool bValidUVTid = Target->AppliedCanonical->Attributes()->GetUVLayer(Target->UVLayerIndex)->IsSetTriangle(ClickedTid);
+
+	if (!bValidUVTid)
+	{
+		return;
+	}
+
+	double TexelDensity = FUVMetrics::TexelDensity(*Target->AppliedCanonical, Target->UVLayerIndex, ClickedTid, Settings->TextureResolution);
+
+	Settings->TargetWorldUnits = Settings->TargetPixelCount / TexelDensity;
+}
+
+
+void UUVEditorTexelDensityTool::Shutdown(EToolShutdownType ShutdownType)
+{
+	Settings->SaveProperties(this);
+
+	if (LivePreviewInputRouter.IsValid())
+	{
+		// TODO: Arguably the live preview input router should do this for us before Shutdown(), but
+		// we don't currently have support for that...
+		LivePreviewInputRouter->ForceTerminateSource(LivePreviewBehaviorSource);
+		LivePreviewInputRouter->DeregisterSource(LivePreviewBehaviorSource);
+	}
+
+	if (ShutdownType == EToolShutdownType::Accept)
+	{
+		UUVToolEmitChangeAPI* ChangeAPI = GetToolManager()->GetContextObjectStore()->FindContext<UUVToolEmitChangeAPI>();
+
+		const FText TransactionName(LOCTEXT("TexelDensityTransactionName", "Texel Density"));
+		ChangeAPI->BeginUndoTransaction(TransactionName);
+		for (const TObjectPtr<UUVEditorToolMeshInput>& Target : Targets)
+		{
+			// Set things up for undo. 
+			FDynamicMeshChangeTracker ChangeTracker(Target->UnwrapCanonical.Get());
+			ChangeTracker.BeginChange();
+
+			for (int32 Tid : Target->UnwrapCanonical->TriangleIndicesItr())
+			{
+				ChangeTracker.SaveTriangle(Tid, true);
+			}
+
+			Target->UpdateCanonicalFromPreviews();
+
+			ChangeAPI->EmitToolIndependentUnwrapCanonicalChange(Target, ChangeTracker.EndChange(), LOCTEXT("ApplyTexelDensityTool", "Texel Density Tool"));
+		}
+
+		ChangeAPI->EndUndoTransaction();
+
+		// Analytics
+		RecordAnalytics();
+	}
+	else
+	{
+		// Reset the inputs
+		for (const TObjectPtr<UUVEditorToolMeshInput>& Target : Targets)
+		{
+			Target->UpdatePreviewsFromCanonical();
+		}
+	}
+
+	for (const TObjectPtr<UUVEditorToolMeshInput>& Target : Targets)
+	{
+		Target->AppliedPreview->ClearOpFactory();
+	}
+	for (int32 FactoryIndex = 0; FactoryIndex < Factories.Num(); ++FactoryIndex)
+	{
+		Factories[FactoryIndex] = nullptr;
+	}
+
+	LivePreviewAPI = nullptr;
+
+	LivePreviewBehaviorSet->RemoveAll();
+	LivePreviewBehaviorSet = nullptr;
+	LivePreviewBehaviorSource = nullptr;
+
+	UnwrapGeometry->Disconnect();
+	LivePreviewGeometry->Disconnect();
+	UnwrapGeometry = nullptr;
+	LivePreviewGeometry = nullptr;
+
+	Spatials2D.Reset();
+	Spatials3D.Reset();
+
+	Settings = nullptr;
+	Targets.Empty();
+}
+
+void UUVEditorTexelDensityTool::OnTick(float DeltaTime)
+{
+	if (PendingAction == ETexelDensityToolAction::BeginSamping)
+	{
+		for (const TObjectPtr<UUVEditorToolMeshInput>& Target : Targets)
+		{
+			Target->UpdatePreviewsFromCanonical();
+			this->UVToolSelectionAPI->RebuildUnwrapHighlight(Target->UnwrapPreview->PreviewMesh->GetTransform());
+		}
+
+		PendingAction = ETexelDensityToolAction::Sampling;
+	}
+	else if( PendingAction == ETexelDensityToolAction::Sampling)
+	{
+		UpdateHover();
+
+		// Apply click if there was one (signaled by a non-InvalidId tid)
+		if (ClickedTid != IndexConstants::InvalidID)
+		{
+			ApplyClick();
+			ClickedTid = IndexConstants::InvalidID;
+
+			PendingAction = ETexelDensityToolAction::NoAction;
+			PerformBackgroundScalingTask();
+		}
+	}
+	else
+	{
+		ClearHover();
+	}
+
+	if (PendingAction == ETexelDensityToolAction::Processing && RemainingTargetsToProcess == 0)
+	{
+		PendingAction = ETexelDensityToolAction::NoAction;
+	}
+
+	for (const TObjectPtr<UUVEditorToolMeshInput>& Target : Targets)
+	{
+		Target->AppliedPreview->Tick(DeltaTime);
+	}
+}
+
+void UUVEditorTexelDensityTool::Render(IToolsContextRenderAPI* RenderAPI)
+{
+
+}
+
+bool UUVEditorTexelDensityTool::CanAccept() const
+{
+	if (UVToolSelectionAPI->HaveSelections())
+	{
+		for (FUVToolSelection Selection : UVToolSelectionAPI->GetSelections())
+		{
+			if (!Selection.Target->AppliedPreview->HaveValidResult())
+			{
+				return false;
+			}
+		}
+	}
+	else
+	{
+		for (const TObjectPtr<UUVEditorToolMeshInput>& Target : Targets)
+		{
+			if (!Target->AppliedPreview->HaveValidResult())
+			{
+				return false;
+			}
+		}
+	}
+	return true;
+}
+
+void UUVEditorTexelDensityTool::OnPropertyModified(UObject* PropertySet, FProperty* Property)
+{
+	if (PendingAction == ETexelDensityToolAction::Sampling)
+	{
+		PendingAction = ETexelDensityToolAction::NoAction;
+	}
+
+	if (PendingAction == ETexelDensityToolAction::NoAction || PendingAction == ETexelDensityToolAction::Processing)
+	{
+		if (PropertySet == Settings)
+		{
+			PerformBackgroundScalingTask();
+		}
+	}
+}
+
+void UUVEditorTexelDensityTool::RequestAction(ETexelDensityToolAction ActionType)
+{
+	if (PendingAction == ETexelDensityToolAction::NoAction)
+	{
+		PendingAction = ActionType;
+	}
+}
+
+ETexelDensityToolAction UUVEditorTexelDensityTool::ActiveAction() const
+{
+	return PendingAction;
+}
+
+void UUVEditorTexelDensityTool::PerformBackgroundScalingTask()
+{
+	PendingAction = ETexelDensityToolAction::Processing;
+
+	RemainingTargetsToProcess = 0;
+	if (UVToolSelectionAPI->HaveSelections())
+	{
+		for (FUVToolSelection Selection : UVToolSelectionAPI->GetSelections())
+		{
+			RemainingTargetsToProcess++;
+			Selection.Target->AppliedPreview->InvalidateResult();
+		}
+	}
+	else
+	{
+		for (int32 TargetIndex = 0; TargetIndex < Targets.Num(); ++TargetIndex)
+		{
+			RemainingTargetsToProcess++;
+			Targets[TargetIndex]->AppliedPreview->InvalidateResult();
+		}
+	}
+
+}
+
+void UUVEditorTexelDensityTool::RecordAnalytics()
+{
+	using namespace UVEditorAnalytics;
+
+	if (!FEngineAnalytics::IsAvailable())
+	{
+		return;
+	}
+
+	TArray<FAnalyticsEventAttribute> Attributes;
+	Attributes.Add(FAnalyticsEventAttribute(TEXT("Timestamp"), FDateTime::UtcNow().ToString()));
+
+	// Tool inputs
+	InputTargetAnalytics.AppendToAttributes(Attributes, "Input");
+
+	// Tool outputs	
+	const FTargetAnalytics OutputTargetAnalytics = CollectTargetAnalytics(Targets);
+	OutputTargetAnalytics.AppendToAttributes(Attributes, "Output");
+
+	// Tool stats
+	if (CanAccept())
+	{
+		TArray<double> PerAssetValidResultComputeTimes;
+		for (const TObjectPtr<UUVEditorToolMeshInput>& Target : Targets)
+		{
+			// Note: This would log -1 if the result was invalid, but checking CanAccept above ensures results are valid
+			PerAssetValidResultComputeTimes.Add(Target->AppliedPreview->GetValidResultComputeTime());
+		}
+		Attributes.Add(FAnalyticsEventAttribute(TEXT("Stats.PerAsset.ComputeTimeSeconds"), PerAssetValidResultComputeTimes));
+	}
+	Attributes.Add(FAnalyticsEventAttribute(TEXT("Stats.ToolActiveDuration"), (FDateTime::UtcNow() - ToolStartTimeAnalytics).ToString()));
+
+	FEngineAnalytics::GetProvider().RecordEvent(UVEditorAnalyticsEventName(TEXT("TexelDensityTool")), Attributes);
+}
+
+#undef LOCTEXT_NAMESPACE
\ No newline at end of file
diff --git a/Engine/Plugins/Editor/UVEditor/Source/UVEditorTools/Public/Operators/UVEditorTexelDensityOp.h b/Engine/Plugins/Editor/UVEditor/Source/UVEditorTools/Public/Operators/UVEditorTexelDensityOp.h
new file mode 100644
--- /dev/null
+++ b/Engine/Plugins/Editor/UVEditor/Source/UVEditorTools/Public/Operators/UVEditorTexelDensityOp.h
@@ -0,0 +1,122 @@
+// Copyright Epic Games, Inc. All Rights Reserved.
+
+#pragma once
+
+#include "CoreMinimal.h"
+#include "Util/ProgressCancel.h"
+#include "ModelingOperators.h"
+#include "InteractiveTool.h"
+#include "DynamicMesh/DynamicMeshAttributeSet.h"
+
+#include "UVEditorTexelDensityOp.generated.h"
+
+UENUM()
+enum class ETexelDensityToolMode
+{
+	ApplyToIslands,
+	ApplyToWhole,
+	Normalize
+};
+
+UCLASS()
+class UVEDITORTOOLS_API UUVEditorTexelDensitySettings : public UInteractiveToolPropertySet
+{
+	GENERATED_BODY()
+public:
+
+	UFUNCTION()
+	virtual bool InSamplingMode() const;
+
+	UPROPERTY(EditAnywhere, Category = "TexelDensity", meta = (DisplayName = "Scale Mode", EditCondition = "!InSamplingMode" ))
+	ETexelDensityToolMode TexelDensityMode = ETexelDensityToolMode::ApplyToIslands;
+
+	UPROPERTY(EditAnywhere, Category = "TexelDensity", meta = (DisplayName = "World Units", EditCondition = "TexelDensityMode != ETexelDensityToolMode::Normalize"))
+	float TargetWorldUnits = 100;
+
+	UPROPERTY(EditAnywhere, Category = "TexelDensity", meta = (DisplayName = "Pixels", EditCondition = "TexelDensityMode != ETexelDensityToolMode::Normalize && !InSamplingMode"))
+	float TargetPixelCount = 1024;
+
+	UPROPERTY(EditAnywhere, Category = "TexelDensity", meta = (DisplayName = "Texture Dimensions", EditCondition = "TexelDensityMode != ETexelDensityToolMode::Normalize && !InSamplingMode", LinearDeltaSensitivity = 1000, Delta = 64, SliderExponent = 1, UIMin = 2, UIMax = 16384))
+	float TextureResolution = 1024;
+
+	UPROPERTY(EditAnywhere, Category = "TexelDensity", meta = (DisplayName = "Use UDIMs", EditCondition = "TexelDensityMode != ETexelDensityToolMode::Normalize && !InSamplingMode", TransientToolProperty))
+	bool bEnableUDIMLayout = false;
+};
+
+
+namespace UE
+{
+	namespace Geometry
+	{
+		class FDynamicMesh3;
+		class FDynamicMeshUVPacker;
+
+		enum class EUVTexelDensityOpModes
+		{
+			ScaleGlobal = 0,
+			ScaleIslands = 1,
+			Normalize = 2
+		};
+
+		class UVEDITORTOOLS_API FUVEditorTexelDensityOp : public FDynamicMeshOperator
+		{
+		public:
+			virtual ~FUVEditorTexelDensityOp() {}
+
+			// inputs
+			TSharedPtr<FDynamicMesh3, ESPMode::ThreadSafe> OriginalMesh;
+
+			EUVTexelDensityOpModes TexelDensityMode = EUVTexelDensityOpModes::ScaleGlobal;
+
+			int UVLayerIndex = 0;
+			int TextureResolution = 1024;
+			int TargetWorldSpaceMeasurement = 100;
+			int TargetPixelCountMeasurement = 1024;
+
+			bool bMaintainOriginatingUDIM = false;
+			TOptional<TSet<int32>> Selection;
+			TOptional<TMap<int32, int32>> TextureResolutionPerUDIM;
+			FVector2f UVTranslation = FVector2f::Zero();
+
+			void SetTransform(const FTransformSRT3d& Transform);
+
+			//
+			// FDynamicMeshOperator implementation
+			// 
+
+			virtual void CalculateResult(FProgressCancel* Progress) override;
+
+		protected:
+			void ExecutePacker(FDynamicMeshUVPacker& Packer);			
+			void ScaleMeshSubRegionByDensity(FDynamicMeshUVOverlay* UVLayer, const TArray<int32>& Tids, TSet<int32>& UVElements, int32 TileResolution);
+
+
+		};
+
+	} // end namespace UE::Geometry
+} // end namespace UE
+
+/**
+ * Can be hooked up to a UMeshOpPreviewWithBackgroundCompute to perform UV Layout operations.
+ *
+ * Inherits from UObject so that it can hold a strong pointer to the settings UObject, which
+ * needs to be a UObject to be displayed in the details panel.
+ */
+UCLASS()
+class UVEDITORTOOLS_API UUVTexelDensityOperatorFactory : public UObject, public UE::Geometry::IDynamicMeshOperatorFactory
+{
+	GENERATED_BODY()
+
+public:
+	// IDynamicMeshOperatorFactory API
+	virtual TUniquePtr<UE::Geometry::FDynamicMeshOperator> MakeNewOperator() override;
+
+	UPROPERTY()
+	TObjectPtr<UUVEditorTexelDensitySettings> Settings;
+
+	TOptional<TSet<int32>> Selection;
+	TSharedPtr<UE::Geometry::FDynamicMesh3, ESPMode::ThreadSafe> OriginalMesh;
+	TUniqueFunction<int32()> GetSelectedUVChannel = []() { return 0; };
+	FTransform TargetTransform;
+	TOptional<TMap<int32, int32>> TextureResolutionPerUDIM;
+};
\ No newline at end of file
diff --git a/Engine/Plugins/Editor/UVEditor/Source/UVEditorTools/Public/UVEditorTexelDensityTool.h b/Engine/Plugins/Editor/UVEditor/Source/UVEditorTools/Public/UVEditorTexelDensityTool.h
new file mode 100644
--- /dev/null
+++ b/Engine/Plugins/Editor/UVEditor/Source/UVEditorTools/Public/UVEditorTexelDensityTool.h
@@ -0,0 +1,211 @@
+// Copyright Epic Games, Inc. All Rights Reserved.
+
+#pragma once
+
+#include "CoreMinimal.h"
+#include "BaseTools/SingleSelectionMeshEditingTool.h"
+#include "DynamicMesh/DynamicMesh3.h"
+#include "PropertySets/PolygroupLayersProperties.h"
+#include "Polygroups/PolygroupSet.h"
+#include "Drawing/UVLayoutPreview.h"
+#include "UVEditorToolAnalyticsUtils.h"
+#include "Selection/UVToolSelectionAPI.h"
+#include "Operators/UVEditorTexelDensityOp.h"
+#include "DynamicMesh/DynamicMeshAABBTree3.h"
+
+#include "UVEditorTexelDensityTool.generated.h"
+
+
+// Forward declarations
+class UDynamicMeshComponent;
+class UMaterialInterface;
+class UMaterialInstanceDynamic;
+class UUVEditorToolMeshInput;
+class UUVEditorUVTransformProperties;
+class UUVTexelDensityOperatorFactory;
+class UUVEditorTexelDensityTool;
+class UUVTool2DViewportAPI;
+class UPreviewGeometry;
+
+
+/**
+ *
+ */
+UCLASS()
+class UVEDITORTOOLS_API UUVEditorTexelDensityToolBuilder : public UInteractiveToolBuilder
+{
+	GENERATED_BODY()
+
+public:
+	virtual bool CanBuildTool(const FToolBuilderState& SceneState) const override;
+	virtual UInteractiveTool* BuildTool(const FToolBuilderState& SceneState) const override;
+
+	// This is a pointer so that it can be updated under the builder without
+	// having to set it in the mode after initializing targets.
+	const TArray<TObjectPtr<UUVEditorToolMeshInput>>* Targets = nullptr;
+};
+
+UENUM()
+enum class ETexelDensityToolAction
+{
+	NoAction,
+	Processing,
+
+	BeginSamping,
+	Sampling
+};
+
+
+UCLASS()
+class UVEDITORTOOLS_API UUVEditorTexelDensityActionSettings : public UInteractiveToolPropertySet
+{
+	GENERATED_BODY()
+public:
+
+	TWeakObjectPtr<UUVEditorTexelDensityTool> ParentTool;
+
+	void Initialize(UUVEditorTexelDensityTool* ParentToolIn) { ParentTool = ParentToolIn; }
+	void PostAction(ETexelDensityToolAction Action);
+
+	UFUNCTION(CallInEditor, Category = Actions)
+	void SampleTexelDensity();
+};
+
+UCLASS()
+class UVEDITORTOOLS_API UUVEditorTexelDensityToolSettings : public UUVEditorTexelDensitySettings
+{
+	GENERATED_BODY()
+public:
+
+	TWeakObjectPtr<UUVEditorTexelDensityTool> ParentTool;
+	void Initialize(UUVEditorTexelDensityTool* ParentToolIn) { ParentTool = ParentToolIn; }
+
+	virtual bool InSamplingMode() const override;
+};
+
+/**
+ * UUVEditorRecomputeUVsTool Recomputes UVs based on existing segmentations of the mesh
+ */
+UCLASS()
+class UVEDITORTOOLS_API UUVEditorTexelDensityTool : public UInteractiveTool, public IUVToolSupportsSelection
+{
+	GENERATED_BODY()
+
+public:
+	virtual void Setup() override;
+	virtual void Shutdown(EToolShutdownType ShutdownType) override;
+
+	virtual void OnTick(float DeltaTime) override;
+
+	virtual void Render(IToolsContextRenderAPI* RenderAPI) override;
+
+	virtual bool HasCancel() const override { return true; }
+	virtual bool HasAccept() const override { return true; }
+	virtual bool CanAccept() const override;
+
+	virtual void OnPropertyModified(UObject* PropertySet, FProperty* Property) override;
+
+
+	/**
+	 * The tool will operate on the meshes given here.
+	 */
+	virtual void SetTargets(const TArray<TObjectPtr<UUVEditorToolMeshInput>>& TargetsIn)
+	{
+		Targets = TargetsIn;
+	}
+
+
+	void RequestAction(ETexelDensityToolAction ActionType);
+	ETexelDensityToolAction ActiveAction() const;
+
+private:
+	UPROPERTY()
+	TArray<TObjectPtr<UUVEditorToolMeshInput>> Targets;
+
+	UPROPERTY()
+	TObjectPtr<UUVEditorTexelDensitySettings> Settings = nullptr;
+
+	UPROPERTY()
+	TObjectPtr<UUVEditorTexelDensityActionSettings> ActionSettings = nullptr;
+
+	ETexelDensityToolAction PendingAction = ETexelDensityToolAction::NoAction;
+
+	UPROPERTY()
+	TObjectPtr<UUVToolSelectionAPI> UVToolSelectionAPI = nullptr;
+
+	UPROPERTY()
+	TObjectPtr<UUVToolLivePreviewAPI> LivePreviewAPI = nullptr;
+
+	//~ For UDIM information access
+	UPROPERTY()
+	TObjectPtr< UUVTool2DViewportAPI> UVTool2DViewportAPI = nullptr;
+
+	UPROPERTY()
+	TObjectPtr<UUVToolEmitChangeAPI> EmitChangeAPI = nullptr;
+
+	UPROPERTY()
+	TArray<TObjectPtr<UUVTexelDensityOperatorFactory>> Factories;
+
+	int32 RemainingTargetsToProcess;
+	void PerformBackgroundScalingTask();
+
+	UPROPERTY()
+	TObjectPtr<UInputBehaviorSet> LivePreviewBehaviorSet = nullptr;
+
+	UPROPERTY()
+	TObjectPtr<ULocalInputBehaviorSource> LivePreviewBehaviorSource = nullptr;
+
+	UPROPERTY()
+	TObjectPtr<UPreviewGeometry> UnwrapGeometry = nullptr;
+
+	UPROPERTY()
+	TObjectPtr<UPreviewGeometry> LivePreviewGeometry = nullptr;
+
+	TWeakObjectPtr<UInputRouter> LivePreviewInputRouter = nullptr;
+
+	TArray<TSharedPtr<UE::Geometry::FDynamicMeshAABBTree3>> Spatials2D; // 1:1 with targets
+	TArray<TSharedPtr<UE::Geometry::FDynamicMeshAABBTree3>> Spatials3D; // 1:1 with targets
+
+	FViewCameraState LivePreviewCameraState;
+
+	UPROPERTY()
+	TObjectPtr<UMaterialInstanceDynamic> TriangleSetMaterial = nullptr;
+
+	TSharedPtr<UE::Geometry::FPolygroupSet, ESPMode::ThreadSafe> ActiveGroupSet;
+	void OnSelectedGroupLayerChanged();
+	void UpdateActiveGroupLayer(bool bUpdateFactories = true);
+
+	int32 Get2DHitTriangle(const FRay& WorldRayIn, int32* IndexOf2DSpatialOut = nullptr);
+	int32 Get3DHitTriangle(const FRay& WorldRayIn, int32* IndexOf2DSpatialOut = nullptr);
+
+	void OnMeshTriangleClicked(int32 Tid, int32 IndexOfMesh, bool bTidIsFromUnwrap);
+	void OnMeshTriangleHovered(int32 Tid, int32 IndexOfMesh, bool bTidIsFromUnwrap);
+	void UpdateHover();
+	/** @param bClearHoverInfo If true, also clears HoverVid, etc in addition to just clearing display. */
+	void ClearHover(bool bClearHoverInfo = true);
+
+	void ApplyClick();
+
+	// Used to remember click info to apply on tick
+	int32 ClickedTid = IndexConstants::InvalidID;
+	int32 ClickedMeshIndex = -1;
+	bool bClickWasInUnwrap = false;
+
+	// Used to remember hover info to apply on tick
+	int32 HoverTid = IndexConstants::InvalidID;
+	int32 HoverMeshIndex = IndexConstants::InvalidID;
+	bool bHoverTidIsFromUnwrap = false;
+	int32 LastHoverTid = IndexConstants::InvalidID;
+	int32 LastHoverMeshIndex = IndexConstants::InvalidID;
+	bool bLastHoverTidWasFromUnwrap = false;
+
+	//
+	// Analytics
+	//
+
+	UE::Geometry::UVEditorAnalytics::FTargetAnalytics InputTargetAnalytics;
+	FDateTime ToolStartTimeAnalytics;
+	void RecordAnalytics();
+
+
+};
