UVEditor: Improved handling of multiple object loading in the UVEditor's 3D viewport.

Two behaviors now exist. When selecting from the Level Editor viewport, relative transforms between all selected objects are preserved in the subsequent UVEditor's viewport. When selecting from the Content Browser, or other sources that don't wrap the selected objects in a Component, the UVEditor will provide an automatic layout based on the object's bounding spheres.

Further, there are two new CVARs:

modeling.UVEditor.EnableLivePreviewArrangement - Enables or Disables the auto arrangement feature when selecting from the Content Browser. This exists to support legacy behavior where users might want the object's baked in transform to position objects in the UVEditor.

modeling.UVEditor.LivePreviewArrangementMode - Selects the mode of the auto arrangement. 0 indicates a line formation. 1 indicates an adaptive circular arrangement.

diff --git a/Engine/Plugins/Editor/UVEditor/Source/UVEditor/Private/UVEditor.cpp b/Engine/Plugins/Editor/UVEditor/Source/UVEditor/Private/UVEditor.cpp
--- a/Engine/Plugins/Editor/UVEditor/Source/UVEditor/Private/UVEditor.cpp
+++ b/Engine/Plugins/Editor/UVEditor/Source/UVEditor/Private/UVEditor.cpp
@@ -27,6 +27,17 @@ void UUVEditor::Initialize(const TArray<TObjectPtr<UObject>>& InObjects)
 	// The asset editor toolkit will do the rest of the necessary initialization inside its PostInitAssetEditor.
 	UAssetEditor::Initialize();
 }
+void UUVEditor::Initialize(const TArray<TObjectPtr<UObject>>& InObjects, const TArray<FTransform3d>& InTransforms)
+{
+	check(InTransforms.Num() == InObjects.Num());
+	ObjectWorldspaceOffsets = InTransforms;
+	Initialize(InObjects);
+}
+
+void UUVEditor::GetWorldspaceRelativeTransforms(TArray<FTransform3d>& OutTransforms)
+{
+	OutTransforms.Append(ObjectWorldspaceOffsets);
+}
 
 IAssetEditorInstance* UUVEditor::GetInstanceInterface() 
 { 
diff --git a/Engine/Plugins/Editor/UVEditor/Source/UVEditor/Private/UVEditorMode.cpp b/Engine/Plugins/Editor/UVEditor/Source/UVEditor/Private/UVEditorMode.cpp
--- a/Engine/Plugins/Editor/UVEditor/Source/UVEditor/Private/UVEditorMode.cpp
+++ b/Engine/Plugins/Editor/UVEditor/Source/UVEditor/Private/UVEditorMode.cpp
@@ -1366,17 +1366,23 @@ void UUVEditorMode::FocusLivePreviewCameraOnSelection()
 
 	for (const FUVToolSelection& Selection : CurrentSelections)
 	{
-		SelectionBoundingBox.Contain(Selection.GetConvertedSelectionForAppliedMesh().ToBoundingBox(*Selection.Target->AppliedCanonical));
+		FTransform3d Transform = Selection.Target->AppliedPreview->PreviewMesh->GetTransform();
+		SelectionBoundingBox.Contain(Selection.GetConvertedSelectionForAppliedMesh().ToBoundingBox(*Selection.Target->AppliedCanonical, Transform));
 	}
 	for (const FUVToolSelection& Selection : CurrentUnsetSelections)
 	{
-		SelectionBoundingBox.Contain(Selection.ToBoundingBox(*Selection.Target->AppliedCanonical));
+		FTransform3d Transform = Selection.Target->AppliedPreview->PreviewMesh->GetTransform();
+		SelectionBoundingBox.Contain(Selection.ToBoundingBox(*Selection.Target->AppliedCanonical, Transform));
 	}
 	if (CurrentSelections.Num() == 0 && CurrentUnsetSelections.Num() == 0)
 	{
 		for (int32 AssetID = 0; AssetID < ToolInputObjects.Num(); ++AssetID)
 		{
-			SelectionBoundingBox.Contain(ToolInputObjects[AssetID]->AppliedCanonical->GetBounds());
+			FTransform3d Transform = ToolInputObjects[AssetID]->AppliedPreview->PreviewMesh->GetTransform();
+			FAxisAlignedBox3d ObjectBoundingBox = ToolInputObjects[AssetID]->AppliedCanonical->GetBounds();
+			ObjectBoundingBox.Max = Transform.TransformPosition(ObjectBoundingBox.Max);
+			ObjectBoundingBox.Min = Transform.TransformPosition(ObjectBoundingBox.Min);
+			SelectionBoundingBox.Contain(ObjectBoundingBox);
 		}
 	}
 	
diff --git a/Engine/Plugins/Editor/UVEditor/Source/UVEditor/Private/UVEditorModule.cpp b/Engine/Plugins/Editor/UVEditor/Source/UVEditor/Private/UVEditorModule.cpp
--- a/Engine/Plugins/Editor/UVEditor/Source/UVEditor/Private/UVEditorModule.cpp
+++ b/Engine/Plugins/Editor/UVEditor/Source/UVEditor/Private/UVEditorModule.cpp
@@ -147,23 +147,30 @@ void FUVEditorModule::RegisterMenus()
 			// We are interested in the (unique) assets backing the actor, or else the actor
 			// itself if it is not asset backed (such as UDynamicMesh).
 			const AActor* Actor = static_cast<AActor*>(*It);
-			TArray<UObject*> ActorAssets;
-			Actor->GetReferencedContentObjects(ActorAssets);
 
-			if (ActorAssets.Num() > 0)
+			// We need to determine if there are components under this actor first,
+			// to ensure we generate Component based targets. Otherwise we can use
+			// the referenced content inside to generate more direct targets. This
+			// is important because otherwise we won't have access to component
+			// level data, such as transforms within the level editor, if we actually
+			// do want it for tools.
+
+			TInlineComponentArray<UActorComponent*> ActorComponents;
+			Actor->GetComponents(ActorComponents);
+
+			if (ActorComponents.Num() > 0)
+			{
+				Result.TargetObjects.Append(ActorComponents);
+			}
+			else
 			{
+				TArray<UObject*> ActorAssets;
+				Actor->GetReferencedContentObjects(ActorAssets);
 				for (UObject* Asset : ActorAssets)
 				{
 					Result.TargetObjects.AddUnique(Asset);
 				}
 			}
-			else
-			{
-				// Need to transform actors to components here because this is what our tool targets expect
-				TInlineComponentArray<UActorComponent*> ActorComponents;
-				Actor->GetComponents(ActorComponents);
-				Result.TargetObjects.Append(ActorComponents);
-			}
 		}
 
 		Result.UVSubsystem = GEditor->GetEditorSubsystem<UUVEditorSubsystem>();
diff --git a/Engine/Plugins/Editor/UVEditor/Source/UVEditor/Private/UVEditorSubsystem.cpp b/Engine/Plugins/Editor/UVEditor/Source/UVEditor/Private/UVEditorSubsystem.cpp
--- a/Engine/Plugins/Editor/UVEditor/Source/UVEditor/Private/UVEditorSubsystem.cpp
+++ b/Engine/Plugins/Editor/UVEditor/Source/UVEditor/Private/UVEditorSubsystem.cpp
@@ -2,10 +2,16 @@
 
 #include "UVEditorSubsystem.h"
 
+#include "CoreMinimal.h"
+
 #include "ToolTargetManager.h"
 #include "ToolTargets/DynamicMeshComponentToolTarget.h"
 #include "ToolTargets/StaticMeshToolTarget.h"
 #include "ToolTargets/SkeletalMeshToolTarget.h"
+#include "ToolTargets/StaticMeshComponentToolTarget.h"
+#include "ToolTargets/SkeletalMeshComponentToolTarget.h"
+
+#include "ModelingToolTargetUtil.h"
 #include "UVEditor.h"
 #include "UVEditorMode.h"
 
@@ -13,6 +19,104 @@
 
 using namespace UE::Geometry;
 
+static TAutoConsoleVariable<int32> CVarEnableLivePreviewArrangement(
+	TEXT("modeling.UVEditor.EnableLivePreviewArrangement"),
+	1,
+	TEXT("Enable auto arranging objects in the UV Editor's live viewport when multiple objects are loaded from the Content Browser."));
+
+static TAutoConsoleVariable<int32> CVarLivePreviewArrangementMode(
+	TEXT("modeling.UVEditor.LivePreviewArrangementMode"),
+	0,
+	TEXT("Pattern to use for UV Editor Live Preview auto arrangement: 0 - Line, 1 - Circular"));
+
+namespace UVSubsystemLocals
+{
+	void ArrangeAsLine(TArray<FTransform3d>& ObjectTransforms, TArray<double>& ObjectBoundingRadii, FVector3d& ObjectsAverageCenter)
+	{
+		double ObjectsMaxRadius = FMath::Max(ObjectBoundingRadii);
+
+		// Increase all radii by a fraction of the average radius for all objects, to provide geometry relative padding between objects.
+		ObjectsMaxRadius = ObjectsMaxRadius * 1.2;
+
+		for (int32 TargetIndex = 0; TargetIndex < ObjectBoundingRadii.Num(); ++TargetIndex)
+		{
+			ObjectTransforms[TargetIndex] = FTransform3d(FVector(0, TargetIndex * ObjectsMaxRadius * 2, 0));
+			ObjectsAverageCenter += FVector(TargetIndex * ObjectsMaxRadius, 0, 0);
+		}
+
+		ObjectsAverageCenter = ObjectsAverageCenter / ObjectBoundingRadii.Num();
+	}
+
+	void ArrangeAsCircles(TArray<FTransform3d>& ObjectTransforms, TArray<double>& ObjectBoundingRadii, FVector3d& ObjectsAverageCenter)
+	{
+		// If we need to initialize the transforms from bounding boxes, we begin a algorithm to sequentially place them in a spiral around the first object.
+		// The algorithm proceeds as follows: 
+		// 1. Place first object at (0,0,0)
+		// 2. Place next object at a minimal touching distance in some direction
+		// 3. Check all previously placed objects for overlap. If there's overlap, push the object along it's placement direction until it's clear
+		// 4. Rotate the placement angle for next iteration by the amount that creates zero overlap between the current placing object and the following one.
+		// 5. Goto step 2, unless there are no remaining objects to place. 
+		//
+		// The algorthim should result in a pattern similar to a hexgonal circle packing with uniform sized spheres, or an approximation of one for variable sized objects.
+
+		double ObjectsMaxRadius = FMath::Max(ObjectBoundingRadii);
+
+		// Increase all radii by a fraction of the average radius for all objects, to provide geometry relative padding between objects.
+		ObjectsMaxRadius = ObjectsMaxRadius * 1.2;
+
+		// TODO: For now, we're using the max radius for all objects.
+		// However, the algorithm below is written for variable radii.
+		// If we decide we like this option, we can simplify the below
+		// for use with a constant radius for all objects.
+		for (int32 Index = 0; Index < ObjectBoundingRadii.Num(); ++Index)
+		{
+			ObjectBoundingRadii[Index] = ObjectsMaxRadius;
+		}
+
+		// "Insert" first object at (0,0,0)
+		ObjectsAverageCenter = FVector3d::ZeroVector;
+		ObjectTransforms[0] = FTransform3d::Identity;
+		double NextPlacementAngle = 0.0;
+
+		for (int32 TargetIndex = 1; TargetIndex < ObjectBoundingRadii.Num(); ++TargetIndex)
+		{
+			// Compute next object placement ray and translation based on next placement angle and naive radii values.
+			FVector3d TestPlacementRay = FVector3d(FMath::Cos(NextPlacementAngle), FMath::Sin(NextPlacementAngle), 0.0);
+			FVector3d TestTranslation = ((ObjectBoundingRadii[0] + ObjectBoundingRadii[TargetIndex]) * TestPlacementRay);
+
+			// Save the current radius of the "ring" we want to use to compute the next angle for.
+			double TestRadius = ObjectBoundingRadii[0] + ObjectBoundingRadii[TargetIndex];
+
+			for (int32 PriorIndex = 0; PriorIndex < TargetIndex; ++PriorIndex)
+			{
+				// Test the collision with prior placed object, moving the translation and radius until no other objects collide
+				double CenterDistance = FVector3d::DistSquaredXY(TestTranslation, ObjectTransforms[PriorIndex].GetTranslation());
+				double TestCollisionPenetration = CenterDistance - FMath::Square(ObjectBoundingRadii[TargetIndex] + ObjectBoundingRadii[PriorIndex]);
+				if (TestCollisionPenetration < 0.0 && !FMath::IsNearlyZero(TestCollisionPenetration))
+				{
+					TestTranslation += TestPlacementRay * 2.0 * ObjectBoundingRadii[PriorIndex];
+					TestRadius += ObjectBoundingRadii[PriorIndex];
+				}
+			}
+
+			ObjectTransforms[TargetIndex] = FTransform3d(TestTranslation);
+			ObjectsAverageCenter += TestTranslation;
+
+			// Finally if there's more objects to place, advance the angle by with a little chord math.
+			// The next position to fit a circle will be a chord on the "ring" the current placement is at. 
+			// It is possible this choice will result in a collision next pass, but the algorithm will automatically
+			// adjust for this as it goes.
+			if (TargetIndex + 1 < ObjectBoundingRadii.Num())
+			{
+				double NextChordDistance = ObjectBoundingRadii[TargetIndex] + ObjectBoundingRadii[TargetIndex + 1];
+				NextPlacementAngle = NextPlacementAngle + (2 * FMath::FastAsin(NextChordDistance / (2 * TestRadius)));
+			}
+		}
+
+		ObjectsAverageCenter = ObjectsAverageCenter / ObjectBoundingRadii.Num();
+	}
+}
+
 void UUVEditorSubsystem::Initialize(FSubsystemCollectionBase& Collection)
 {
 	// The subsystem has its own tool target manager because it must exist before any UV editors exist,
@@ -23,6 +127,8 @@ void UUVEditorSubsystem::Initialize(FSubsystemCollectionBase& Collection)
 	ToolTargetManager->AddTargetFactory(NewObject<UStaticMeshToolTargetFactory>(ToolTargetManager));
 	ToolTargetManager->AddTargetFactory(NewObject<USkeletalMeshToolTargetFactory>(ToolTargetManager));
 	ToolTargetManager->AddTargetFactory(NewObject<UDynamicMeshComponentToolTargetFactory>(ToolTargetManager));
+	ToolTargetManager->AddTargetFactory(NewObject<UStaticMeshComponentToolTargetFactory>(ToolTargetManager));
+	ToolTargetManager->AddTargetFactory(NewObject<USkeletalMeshComponentToolTargetFactory>(ToolTargetManager));
 }
 
 void UUVEditorSubsystem::Deinitialize()
@@ -124,11 +230,76 @@ void UUVEditorSubsystem::StartUVEditor(TArray<TObjectPtr<UObject>> ObjectsToEdit
 		return;
 	}
 
+	TArray<FTransform3d> ObjectTransforms;
+	TArray<double> ObjectBoundingRadii;
+	TArray<UToolTarget*> Targets;
+	FVector3d ObjectsAverageCenter = FVector3d::ZeroVector;
+
+	bool bInitializeTransformsFromBoundingBoxes = true;
+
+	// TODO: The below logic needs to be moved out of the Subsystem and into the UV Editor proper when we
+	// have a better understanding of how we want to integrate the 3D arrangement logic into the UX of the editor.
+
+	for (UObject* Object : ObjectsToEdit)
+	{
+		UToolTarget* Target = ToolTargetManager->BuildTarget(Object, UUVEditorMode::GetToolTargetRequirements());
+		Targets.Add(Target);
+
+
+		// There are two options here: A) all targets are components (selected from the viewport) or B) all targets aren't components (selected from the content browser). 
+		// TODO: Check if it's possible to have a mixture?
+		// Depending on which scenario we're in, we'll either use the component transforms relative to each other, or build some transforms ourselves.
+		UPrimitiveComponent* Component = UE::ToolTarget::GetTargetComponent(Target);
+		if (Component)
+		{
+			bInitializeTransformsFromBoundingBoxes = false;
+			FTransform3d Transform = UE::ToolTarget::GetLocalToWorldTransform(Target);
+			ObjectTransforms.Add(Transform);
+			ObjectsAverageCenter += Transform.GetTranslation();
+		}
+		else
+		{
+			FDynamicMesh3 Mesh = UE::ToolTarget::GetDynamicMeshCopy(Target);
+			double MaxXYRadius = FMath::Max(Mesh.GetBounds().Dimension(0), Mesh.GetBounds().Dimension(1)) * 0.5;
+			ObjectBoundingRadii.Add(MaxXYRadius);
+			ObjectTransforms.Add(FTransform3d::Identity);
+		}
+	}	
+
+
+	bool bEnableAutoArrangement = (CVarEnableLivePreviewArrangement.GetValueOnGameThread() > 0);
+	if (bEnableAutoArrangement && bInitializeTransformsFromBoundingBoxes && Targets.Num() > 0)
+	{
+		switch (CVarLivePreviewArrangementMode.GetValueOnGameThread())
+		{
+		case 0: 
+			UVSubsystemLocals::ArrangeAsLine(ObjectTransforms, ObjectBoundingRadii, ObjectsAverageCenter);
+			break;
+		case 1:
+			UVSubsystemLocals::ArrangeAsCircles(ObjectTransforms, ObjectBoundingRadii, ObjectsAverageCenter);
+			break;
+		default:
+			ensureMsgf(false, TEXT("No valid arrangement CVAR selected for UV Editor auto-arrange. Defaulting to no arrangement (legacy behavior).") );
+			break;
+		}
+	}
+	else
+	{
+		ObjectsAverageCenter = ObjectsAverageCenter / ObjectsToEdit.Num();
+	}
+	
+	
+	for (FTransform3d& Transform: ObjectTransforms)
+	{
+		Transform.SetTranslation(Transform.GetTranslation() - ObjectsAverageCenter);
+	}
+
+
 	UUVEditor* UVEditor = NewObject<UUVEditor>(this);
 
 	// Among other things, this call registers the UV editor with the asset editor subsystem,
 	// which will prevent it from being garbage collected.
-	UVEditor->Initialize(ObjectsToEdit);
+	UVEditor->Initialize(ObjectsToEdit, ObjectTransforms);
 
 	for (TObjectPtr<UObject>& Object : ObjectsToEdit)
 	{
diff --git a/Engine/Plugins/Editor/UVEditor/Source/UVEditor/Private/UVEditorToolkit.cpp b/Engine/Plugins/Editor/UVEditor/Source/UVEditor/Private/UVEditorToolkit.cpp
--- a/Engine/Plugins/Editor/UVEditor/Source/UVEditor/Private/UVEditorToolkit.cpp
+++ b/Engine/Plugins/Editor/UVEditor/Source/UVEditor/Private/UVEditorToolkit.cpp
@@ -443,7 +443,7 @@ void FUVEditorToolkit::PostInitAssetEditor()
 
 	// TODO: get these when possible (from level editor selection, for instance), and set them to something reasonable otherwise.
 	TArray<FTransform> ObjectTransforms;
-	ObjectTransforms.SetNum(ObjectsToEdit.Num());
+	Cast<UUVEditor>(OwningAssetEditor)->GetWorldspaceRelativeTransforms(ObjectTransforms);
 
 	// This static method call initializes the variety of contexts that UVEditorMode needs to be available in
 	// the context store on Enter() to function properly.
diff --git a/Engine/Plugins/Editor/UVEditor/Source/UVEditor/Public/UVEditor.h b/Engine/Plugins/Editor/UVEditor/Source/UVEditor/Public/UVEditor.h
--- a/Engine/Plugins/Editor/UVEditor/Source/UVEditor/Public/UVEditor.h
+++ b/Engine/Plugins/Editor/UVEditor/Source/UVEditor/Public/UVEditor.h
@@ -38,6 +38,9 @@ class UVEDITOR_API UUVEditor : public UAssetEditor
 
 public:
 	void Initialize(const TArray<TObjectPtr<UObject>>& InObjects);
+	void Initialize(const TArray<TObjectPtr<UObject>>& InObjects, const TArray<FTransform3d>& InTransforms);
+
+	void GetWorldspaceRelativeTransforms(TArray<FTransform3d>& OutTransforms);
 
 	/** Returns the asset editor instance interface, so that its window can be focused, for example. */
 	IAssetEditorInstance* GetInstanceInterface();
@@ -54,5 +57,8 @@ protected:
 	UPROPERTY()
 	TArray<TObjectPtr<UObject>> OriginalObjectsToEdit;
 
+	UPROPERTY()
+	TArray<FTransform3d> ObjectWorldspaceOffsets;
+
 	friend class UUVEditorSubsystem;
 };
diff --git a/Engine/Plugins/Editor/UVEditor/Source/UVEditorTools/Private/Selection/UVEditorMeshSelectionMechanic.cpp b/Engine/Plugins/Editor/UVEditor/Source/UVEditorTools/Private/Selection/UVEditorMeshSelectionMechanic.cpp
--- a/Engine/Plugins/Editor/UVEditor/Source/UVEditorTools/Private/Selection/UVEditorMeshSelectionMechanic.cpp
+++ b/Engine/Plugins/Editor/UVEditor/Source/UVEditorTools/Private/Selection/UVEditorMeshSelectionMechanic.cpp
@@ -1102,8 +1102,22 @@ bool UUVEditorMeshSelectionMechanic::GetHitTid(const FInputDeviceRay& ClickPos,
 	int32& TidOut, int32& AssetIDOut, bool bUseUnwrap, int32* ExistingSelectionObjectIndexOut)
 {
 	auto RayCastSpatial = [this, &ClickPos, &TidOut, &AssetIDOut, &bUseUnwrap](int32 AssetID) {
+		FTransform3d TargetTransform;
+		if (bUseUnwrap)
+		{
+			TargetTransform = Targets[AssetID]->UnwrapPreview->PreviewMesh->GetTransform();
+		}
+		else
+		{
+			TargetTransform = Targets[AssetID]->AppliedPreview->PreviewMesh->GetTransform();
+		}
+
+		FRay3d LocalRay(
+			(FVector3d)TargetTransform.InverseTransformPosition(ClickPos.WorldRay.Origin),
+			(FVector3d)TargetTransform.InverseTransformVector(ClickPos.WorldRay.Direction));
+
 		double RayT = 0;
-		if (GetMeshSpatial(AssetID, bUseUnwrap)->FindNearestHitTriangle(ClickPos.WorldRay, RayT, TidOut))
+		if (GetMeshSpatial(AssetID, bUseUnwrap)->FindNearestHitTriangle(LocalRay, RayT, TidOut))
 		{
 			AssetIDOut = AssetID;
 			return true;
@@ -1491,6 +1505,7 @@ bool UUVEditorMeshSelectionMechanic::OnUpdateHover(const FInputDeviceRay& Device
 	}
 	FDynamicMesh3* UnwrapMesh = Targets[AssetID]->UnwrapCanonical.Get();
 	FDynamicMesh3* AppliedMesh = Targets[AssetID]->AppliedCanonical.Get();
+	FTransform3d AppliedTargetTransform = Targets[AssetID]->AppliedPreview->PreviewMesh->GetTransform();
 	FDynamicMesh3* Mesh = bSourceIsLivePreview ? AppliedMesh : UnwrapMesh;
 	FDynamicMeshUVOverlay& UVOverlay = *AppliedMesh->Attributes()->GetUVLayer(Targets[AssetID]->UVLayerIndex);
 
@@ -1513,11 +1528,11 @@ bool UUVEditorMeshSelectionMechanic::OnUpdateHover(const FInputDeviceRay& Device
 	}
 
 
-	auto SetupVertexHighlight = [](TArray<int32>& IDs, FDynamicMesh3& Mesh, UPointSetComponent& PointSet)
+	auto SetupVertexHighlight = [](TArray<int32>& IDs, FDynamicMesh3& Mesh, UPointSetComponent& PointSet, FTransform3d Transform = FTransform3d::Identity)
 	{
 		for (int32 ID : IDs)
 		{
-			const FVector3d& P = Mesh.GetVertexRef(ID);
+			const FVector3d P = Transform.TransformPosition(Mesh.GetVertexRef(ID));
 			const FRenderablePoint PointToRender(P,
 				FUVEditorUXSettings::SelectionHoverTriangleWireframeColor,
 				FUVEditorUXSettings::SelectionPointThickness);
@@ -1526,13 +1541,13 @@ bool UUVEditorMeshSelectionMechanic::OnUpdateHover(const FInputDeviceRay& Device
 		}
 	};
 
-	auto SetupEdgeHighlight = [](TArray<int32>& IDs, FDynamicMesh3& Mesh, ULineSetComponent& LineSet)
+	auto SetupEdgeHighlight = [](TArray<int32>& IDs, FDynamicMesh3& Mesh, ULineSetComponent& LineSet, FTransform3d Transform = FTransform3d::Identity)
 	{
 		for (int32 ID : IDs)
 		{
 			const FIndex2i EdgeVids = Mesh.GetEdgeV(ID);
-			const FVector& A = Mesh.GetVertexRef(EdgeVids.A);
-			const FVector& B = Mesh.GetVertexRef(EdgeVids.B);
+			const FVector A = Transform.TransformPosition(Mesh.GetVertexRef(EdgeVids.A));
+			const FVector B = Transform.TransformPosition(Mesh.GetVertexRef(EdgeVids.B));
 
 			LineSet.AddLine(A, B,
 				FUVEditorUXSettings::SelectionHoverTriangleWireframeColor,
@@ -1541,7 +1556,7 @@ bool UUVEditorMeshSelectionMechanic::OnUpdateHover(const FInputDeviceRay& Device
 		}
 	};
 
-	auto SetupTriangleHighlight = [this, &UnwrapMesh, &Tid](FDynamicMesh3& Mesh, ULineSetComponent& LineSet, UTriangleSetComponent& TriangleSet)
+	auto SetupTriangleHighlight = [this, &UnwrapMesh, &Tid](FDynamicMesh3& Mesh, ULineSetComponent& LineSet, UTriangleSetComponent& TriangleSet, FTransform3d Transform = FTransform3d::Identity)
 	{
 		if (!UnwrapMesh->IsTriangle(Tid))
 		{
@@ -1549,9 +1564,9 @@ bool UUVEditorMeshSelectionMechanic::OnUpdateHover(const FInputDeviceRay& Device
 		}
 
 		const FIndex3i Vids = Mesh.GetTriangle(Tid);
-		const FVector& A = Mesh.GetVertex(Vids[0]);
-		const FVector& B = Mesh.GetVertex(Vids[1]);
-		const FVector& C = Mesh.GetVertex(Vids[2]);
+		const FVector A = Transform.TransformPosition(Mesh.GetVertex(Vids[0]));
+		const FVector B = Transform.TransformPosition(Mesh.GetVertex(Vids[1]));
+		const FVector C = Transform.TransformPosition(Mesh.GetVertex(Vids[2]));
 
 		LineSet.AddLine(A, B, FUVEditorUXSettings::SelectionHoverTriangleWireframeColor,
 			FUVEditorUXSettings::SelectionLineThickness, FUVEditorUXSettings::SelectionHoverWireframeDepthBias);
@@ -1570,19 +1585,19 @@ bool UUVEditorMeshSelectionMechanic::OnUpdateHover(const FInputDeviceRay& Device
 	if (SelectionMode == ESelectionMode::Vertex)
 	{
 		SetupVertexHighlight(Converted2DIDs, *UnwrapMesh, *HoverPointSet);
-		SetupVertexHighlight(Converted3DIDs, *AppliedMesh, *LivePreviewHoverPointSet);
-		SetupVertexHighlight(Unset3DIDs, *AppliedMesh, *LivePreviewHoverPointSet);
+		SetupVertexHighlight(Converted3DIDs, *AppliedMesh, *LivePreviewHoverPointSet, AppliedTargetTransform);
+		SetupVertexHighlight(Unset3DIDs, *AppliedMesh, *LivePreviewHoverPointSet, AppliedTargetTransform);
 	}
 	else if (SelectionMode == ESelectionMode::Edge)
 	{
 		SetupEdgeHighlight(Converted2DIDs, *UnwrapMesh, *HoverLineSet);
-		SetupEdgeHighlight(Converted3DIDs, *AppliedMesh, *LivePreviewHoverLineSet);
-		SetupEdgeHighlight(Unset3DIDs, *AppliedMesh, *LivePreviewHoverLineSet);
+		SetupEdgeHighlight(Converted3DIDs, *AppliedMesh, *LivePreviewHoverLineSet, AppliedTargetTransform);
+		SetupEdgeHighlight(Unset3DIDs, *AppliedMesh, *LivePreviewHoverLineSet, AppliedTargetTransform);
 	}
 	else
 	{
 		SetupTriangleHighlight(*UnwrapMesh, *HoverLineSet, *HoverTriangleSet);
-		SetupTriangleHighlight(*AppliedMesh, *LivePreviewHoverLineSet, *LivePreviewHoverTriangleSet);
+		SetupTriangleHighlight(*AppliedMesh, *LivePreviewHoverLineSet, *LivePreviewHoverTriangleSet, AppliedTargetTransform);
 	}
 
 	return true;
diff --git a/Engine/Plugins/Editor/UVEditor/Source/UVEditorTools/Private/Selection/UVToolSelection.cpp b/Engine/Plugins/Editor/UVEditor/Source/UVEditorTools/Private/Selection/UVToolSelection.cpp
--- a/Engine/Plugins/Editor/UVEditor/Source/UVEditorTools/Private/Selection/UVToolSelection.cpp
+++ b/Engine/Plugins/Editor/UVEditor/Source/UVEditorTools/Private/Selection/UVToolSelection.cpp
@@ -119,7 +119,7 @@ bool FUVToolSelection::IsAllSelected(const FDynamicMesh3& Mesh) const
 	return AreElementsPresentInMesh(Mesh);
 }
 
-FAxisAlignedBox3d FUVToolSelection::ToBoundingBox(const FDynamicMesh3& Mesh) const
+FAxisAlignedBox3d FUVToolSelection::ToBoundingBox(const FDynamicMesh3& Mesh, const FTransform3d Transform) const
 {
 	FAxisAlignedBox3d BoundingBox;
 
@@ -130,7 +130,7 @@ FAxisAlignedBox3d FUVToolSelection::ToBoundingBox(const FDynamicMesh3& Mesh) con
 		{
 			if (Mesh.IsVertex(Vid))
 			{
-				BoundingBox.Contain(Mesh.GetVertexRef(Vid));
+				BoundingBox.Contain(Transform.TransformPosition(Mesh.GetVertexRef(Vid)));
 			}
 		}
 		break;
@@ -139,8 +139,8 @@ FAxisAlignedBox3d FUVToolSelection::ToBoundingBox(const FDynamicMesh3& Mesh) con
 		{
 			if (Mesh.IsEdge(Eid))
 			{
-				BoundingBox.Contain(Mesh.GetVertexRef(Mesh.GetEdgeRef(Eid).Vert.A));
-				BoundingBox.Contain(Mesh.GetVertexRef(Mesh.GetEdgeRef(Eid).Vert.B));
+				BoundingBox.Contain(Transform.TransformPosition(Mesh.GetVertexRef(Mesh.GetEdgeRef(Eid).Vert.A)));
+				BoundingBox.Contain(Transform.TransformPosition(Mesh.GetVertexRef(Mesh.GetEdgeRef(Eid).Vert.B)));
 			}
 		}
 		break;
@@ -149,9 +149,9 @@ FAxisAlignedBox3d FUVToolSelection::ToBoundingBox(const FDynamicMesh3& Mesh) con
 		{
 			if (Mesh.IsTriangle(Tid))
 			{
-				BoundingBox.Contain(Mesh.GetVertexRef(Mesh.GetTriangleRef(Tid).A));
-				BoundingBox.Contain(Mesh.GetVertexRef(Mesh.GetTriangleRef(Tid).B));
-				BoundingBox.Contain(Mesh.GetVertexRef(Mesh.GetTriangleRef(Tid).C));
+				BoundingBox.Contain(Transform.TransformPosition(Mesh.GetVertexRef(Mesh.GetTriangleRef(Tid).A)));
+				BoundingBox.Contain(Transform.TransformPosition(Mesh.GetVertexRef(Mesh.GetTriangleRef(Tid).B)));
+				BoundingBox.Contain(Transform.TransformPosition(Mesh.GetVertexRef(Mesh.GetTriangleRef(Tid).C)));
 			}
 		}
 		break;
diff --git a/Engine/Plugins/Editor/UVEditor/Source/UVEditorTools/Private/UVEditorLayoutTool.cpp b/Engine/Plugins/Editor/UVEditor/Source/UVEditorTools/Private/UVEditorLayoutTool.cpp
--- a/Engine/Plugins/Editor/UVEditor/Source/UVEditorTools/Private/UVEditorLayoutTool.cpp
+++ b/Engine/Plugins/Editor/UVEditor/Source/UVEditorTools/Private/UVEditorLayoutTool.cpp
@@ -175,9 +175,9 @@ void UUVEditorLayoutTool::Shutdown(EToolShutdownType ShutdownType)
 		Target->AppliedPreview->ClearOpFactory();
 	}
 
-	for (int32 TargetIndex = 0; TargetIndex < Targets.Num(); ++TargetIndex)
+	for (int32 FactoryIndex = 0; FactoryIndex < Factories.Num(); ++FactoryIndex)
 	{
-		Factories[TargetIndex] = nullptr;
+		Factories[FactoryIndex] = nullptr;
 	}
 
 	Settings = nullptr;
@@ -204,11 +204,24 @@ void UUVEditorLayoutTool::OnPropertyModified(UObject* PropertySet, FProperty* Pr
 
 bool UUVEditorLayoutTool::CanAccept() const
 {
-	for (const TObjectPtr<UUVEditorToolMeshInput>& Target : Targets)
+	if (UVToolSelectionAPI->HaveSelections())
 	{
-		if (!Target->AppliedPreview->HaveValidResult())
+		for (FUVToolSelection Selection : UVToolSelectionAPI->GetSelections())
 		{
-			return false;
+			if (!Selection.Target->AppliedPreview->HaveValidResult())
+			{
+				return false;
+			}
+		}
+	}
+	else
+	{
+		for (const TObjectPtr<UUVEditorToolMeshInput>& Target : Targets)
+		{
+			if (!Target->AppliedPreview->HaveValidResult())
+			{
+				return false;
+			}
 		}
 	}
 	return true;
diff --git a/Engine/Plugins/Editor/UVEditor/Source/UVEditorTools/Private/UVEditorRecomputeUVsTool.cpp b/Engine/Plugins/Editor/UVEditor/Source/UVEditorTools/Private/UVEditorRecomputeUVsTool.cpp
--- a/Engine/Plugins/Editor/UVEditor/Source/UVEditorTools/Private/UVEditorRecomputeUVsTool.cpp
+++ b/Engine/Plugins/Editor/UVEditor/Source/UVEditorTools/Private/UVEditorRecomputeUVsTool.cpp
@@ -224,12 +224,27 @@ void UUVEditorRecomputeUVsTool::Render(IToolsContextRenderAPI* RenderAPI)
 
 bool UUVEditorRecomputeUVsTool::CanAccept() const
 {
-	bool bPreviewsHaveValidResults = true;
-	for (const TObjectPtr<UUVEditorToolMeshInput>& Target : Targets)
+	if (UVToolSelectionAPI->HaveSelections())
 	{
-		bPreviewsHaveValidResults = bPreviewsHaveValidResults && Target->AppliedPreview->HaveValidResult();
+		for (FUVToolSelection Selection : UVToolSelectionAPI->GetSelections())
+		{
+			if (!Selection.Target->AppliedPreview->HaveValidResult())
+			{
+				return false;
+			}
+		}
+	}
+	else
+	{
+		for (const TObjectPtr<UUVEditorToolMeshInput>& Target : Targets)
+		{
+			if (!Target->AppliedPreview->HaveValidResult())
+			{
+				return false;
+			}
+		}
 	}
-	return bPreviewsHaveValidResults;
+	return true;
 }
 
 
diff --git a/Engine/Plugins/Editor/UVEditor/Source/UVEditorTools/Private/UVEditorTransformTool.cpp b/Engine/Plugins/Editor/UVEditor/Source/UVEditorTools/Private/UVEditorTransformTool.cpp
--- a/Engine/Plugins/Editor/UVEditor/Source/UVEditorTools/Private/UVEditorTransformTool.cpp
+++ b/Engine/Plugins/Editor/UVEditor/Source/UVEditorTools/Private/UVEditorTransformTool.cpp
@@ -188,7 +188,7 @@ void UUVEditorTransformTool::Setup()
 		int32 TargetIndex = Targets.Find(&Target);
 
 		TObjectPtr<UUVEditorUVTransformOperatorFactory> Factory = NewObject<UUVEditorUVTransformOperatorFactory>();
-		Factory->TargetTransform = Target.AppliedPreview->PreviewMesh->GetTransform();
+		Factory->TargetTransform = Target.UnwrapPreview->PreviewMesh->GetTransform();
 		Factory->Settings = Settings;
 		Factory->TransformType = ToolMode.Get(EUVEditorUVTransformType::Transform);
 		Factory->OriginalMesh = Target.UnwrapCanonical;
@@ -222,10 +222,10 @@ void UUVEditorTransformTool::Setup()
 			Target.UpdateAppliedPreviewFromUnwrapPreview();
 			
 
-			this->UVToolSelectionAPI->RebuildUnwrapHighlight(Preview->PreviewMesh->GetTransform());
+			this->UVToolSelectionAPI->RebuildUnwrapHighlight(Target.UnwrapPreview->PreviewMesh->GetTransform());
 			});
 
-		Target.UnwrapPreview->OnOpCompleted.AddLambda(
+		Target.UnwrapPreview->OnOpCompleted.AddWeakLambda(this,
 			[this, TargetIndex](const FDynamicMeshOperator* Op)
 			{
 				const FUVEditorUVTransformBaseOp* TransformBaseOp = (const FUVEditorUVTransformBaseOp*)(Op);
@@ -268,6 +268,7 @@ void UUVEditorTransformTool::Shutdown(EToolShutdownType ShutdownType)
 	for (const TObjectPtr<UUVEditorToolMeshInput>& Target : Targets)
 	{
 		Target->UnwrapPreview->OnMeshUpdated.RemoveAll(this);
+		Target->UnwrapPreview->OnOpCompleted.RemoveAll(this);
 	}
 
 	if (ShutdownType == EToolShutdownType::Accept)
@@ -350,11 +351,25 @@ void UUVEditorTransformTool::OnPropertyModified(UObject* PropertySet, FProperty*
 
 bool UUVEditorTransformTool::CanAccept() const
 {
-	for (const TObjectPtr<UUVEditorToolMeshInput>& Target : Targets)
+
+	if (UVToolSelectionAPI->HaveSelections())
+	{
+		for (FUVToolSelection Selection : UVToolSelectionAPI->GetSelections())
+		{
+			if (!Selection.Target->UnwrapPreview->HaveValidResult())
+			{
+				return false;
+			}
+		}
+	}
+	else
 	{
-		if (!Target->UnwrapPreview->HaveValidResult())
+		for (const TObjectPtr<UUVEditorToolMeshInput>& Target : Targets)
 		{
-			return false;
+			if (!Target->UnwrapPreview->HaveValidResult())
+			{
+				return false;
+			}
 		}
 	}
 	return true;
diff --git a/Engine/Plugins/Editor/UVEditor/Source/UVEditorTools/Public/Selection/UVToolSelection.h b/Engine/Plugins/Editor/UVEditor/Source/UVEditorTools/Public/Selection/UVToolSelection.h
--- a/Engine/Plugins/Editor/UVEditor/Source/UVEditorTools/Public/Selection/UVToolSelection.h
+++ b/Engine/Plugins/Editor/UVEditor/Source/UVEditorTools/Public/Selection/UVToolSelection.h
@@ -70,7 +70,7 @@ public:
 
 	bool AreElementsPresentInMesh(const FDynamicMesh3& Mesh) const;
 
-	FAxisAlignedBox3d ToBoundingBox(const FDynamicMesh3& Mesh) const;
+	FAxisAlignedBox3d ToBoundingBox(const FDynamicMesh3& Mesh, const FTransform Transform = FTransform3d::Identity) const;
 
 	/**
 	 * Utility method to construct a new selection from an existing reference selection, only with a new geometry type.
diff --git a/Engine/Plugins/Runtime/MeshModelingToolset/Source/ModelingComponents/Private/ModelingToolTargetUtil.cpp b/Engine/Plugins/Runtime/MeshModelingToolset/Source/ModelingComponents/Private/ModelingToolTargetUtil.cpp
--- a/Engine/Plugins/Runtime/MeshModelingToolset/Source/ModelingComponents/Private/ModelingToolTargetUtil.cpp
+++ b/Engine/Plugins/Runtime/MeshModelingToolset/Source/ModelingComponents/Private/ModelingToolTargetUtil.cpp
@@ -51,7 +51,6 @@ UPrimitiveComponent* UE::ToolTarget::GetTargetComponent(UToolTarget* Target)
 	{
 		return TargetComponent->GetOwnerComponent();
 	}
-	ensure(false);
 	return nullptr;
 }
 
@@ -117,7 +116,6 @@ FTransform3d UE::ToolTarget::GetLocalToWorldTransform(UToolTarget* Target)
 	{
 		return (FTransform3d)TargetComponent->GetWorldTransform();
 	}
-	ensure(false);
 	return FTransform3d();
 }
 
