UVEditor: Fix issue where unwrap material was set with a depth offset multiplier that rendered mesh lines difficult to see in the unwrap view of the editor.

#virtualized

diff --git a/Engine/Plugins/Editor/UVEditor/Source/UVEditor/Private/UVEditorMode.cpp b/Engine/Plugins/Editor/UVEditor/Source/UVEditor/Private/UVEditorMode.cpp
--- a/Engine/Plugins/Editor/UVEditor/Source/UVEditor/Private/UVEditorMode.cpp
+++ b/Engine/Plugins/Editor/UVEditor/Source/UVEditor/Private/UVEditorMode.cpp
@@ -1074,6 +1074,8 @@ void UUVEditorMode::InitializeTargets()
 		AppliedPreviews.Add(AppliedPreview);
 	}
 
+	UMaterialInterface* Material = LoadObject<UMaterial>(nullptr, TEXT("/UVEditor/Materials/UVEditor_UnwrapMaterial"));
+
 	// When creating UV unwraps, these functions will determine the mapping between UV values and the
 	// resulting unwrap mesh vertex positions. 
 	// If we're looking down on the unwrapped mesh, with the Z axis towards us, we want U's to be right, and
@@ -1104,12 +1106,15 @@ void UUVEditorMode::InitializeTargets()
 			ToolInputObject->AppliedPreview->PreviewMesh->SetTransform(Transforms[AssetID]);
 		}
 
+		UMaterialInstanceDynamic* MatInstance = UMaterialInstanceDynamic::Create(Material, GetToolManager());
+		MatInstance->SetVectorParameterValue(TEXT("Color"), FUVEditorUXSettings::GetTriangleColorByTargetIndex(AssetID));
+		MatInstance->SetScalarParameterValue(TEXT("DepthBias"), FUVEditorUXSettings::UnwrapTriangleDepthOffset);
+		MatInstance->SetScalarParameterValue(TEXT("Opacity"), FUVEditorUXSettings::UnwrapTriangleOpacity);
+		MatInstance->SetScalarParameterValue(TEXT("UseVertexColors"), false);
+
+		ToolInputObject->bEnableTriangleVertexColors = false;
 		ToolInputObject->UnwrapPreview->PreviewMesh->SetMaterial(
-			0, ToolSetupUtil::GetCustomTwoSidedDepthOffsetMaterial(
-				GetToolManager(),
-				FUVEditorUXSettings::GetTriangleColorByTargetIndex(AssetID),
-				FUVEditorUXSettings::UnwrapTriangleDepthOffset,
-				FUVEditorUXSettings::UnwrapTriangleOpacity));
+			0, MatInstance);
 
 		// Set up the wireframe display of the unwrapped mesh.
 		UMeshElementsVisualizer* WireframeDisplay = NewObject<UMeshElementsVisualizer>(this);
@@ -1125,7 +1130,8 @@ void UUVEditorMode::InitializeTargets()
 		WireframeDisplay->Settings->bShowNormalSeams = false;
 		// These are not exposed at the visualizer level yet
 		// TODO: Should they be?
-		WireframeDisplay->WireframeComponent->BoundaryEdgeThickness = 2;
+		WireframeDisplay->WireframeComponent->WireframeThickness = FUVEditorUXSettings::WireframeThickness;
+		WireframeDisplay->WireframeComponent->BoundaryEdgeThickness = FUVEditorUXSettings::BoundaryEdgeThickness;
 
 		// The wireframe will track the unwrap preview mesh
 		WireframeDisplay->SetMeshAccessFunction([ToolInputObject](UMeshElementsVisualizer::ProcessDynamicMeshFunc ProcessFunc) {
@@ -1322,7 +1328,7 @@ void UUVEditorMode::UpdateTriangleMaterialBasedOnDisplaySettings()
 		for (int32 AssetID = 0; AssetID < ToolInputObjects.Num(); ++AssetID) {
 			UMaterialInstanceDynamic* MatInstance = UMaterialInstanceDynamic::Create(Material, GetToolManager());
 			MatInstance->SetVectorParameterValue(TEXT("Color"), FUVEditorUXSettings::GetTriangleColorByTargetIndex(AssetID));
-			MatInstance->SetScalarParameterValue(TEXT("PercentDepthOffset"), FUVEditorUXSettings::UnwrapTriangleDepthOffset);
+			MatInstance->SetScalarParameterValue(TEXT("DepthBias"), FUVEditorUXSettings::UnwrapTriangleDepthOffset);
 			MatInstance->SetScalarParameterValue(TEXT("Opacity"), TriangleOpacity);
 			MatInstance->SetScalarParameterValue(TEXT("UseVertexColors"), bUseVertexColors);
 
diff --git a/Engine/Plugins/Editor/UVEditor/Source/UVEditorTools/Private/UVEditorUXSettings.cpp b/Engine/Plugins/Editor/UVEditor/Source/UVEditorTools/Private/UVEditorUXSettings.cpp
--- a/Engine/Plugins/Editor/UVEditor/Source/UVEditorTools/Private/UVEditorUXSettings.cpp
+++ b/Engine/Plugins/Editor/UVEditor/Source/UVEditorTools/Private/UVEditorUXSettings.cpp
@@ -82,7 +82,8 @@ const float FUVEditorUXSettings::SelectionPointThickness(6);
 const float FUVEditorUXSettings::SewLineHighlightThickness(3.0f);
 const float FUVEditorUXSettings::AxisThickness(2.0);
 const float FUVEditorUXSettings::GridMajorThickness(1.0);
-
+const float FUVEditorUXSettings::WireframeThickness(2.0);
+const float FUVEditorUXSettings::BoundaryEdgeThickness(2.0);
 const float FUVEditorUXSettings::ToolPointSize(6);
 const float FUVEditorUXSettings::PivotLineThickness(1.5);
 
diff --git a/Engine/Plugins/Editor/UVEditor/Source/UVEditorTools/Public/UVEditorUXSettings.h b/Engine/Plugins/Editor/UVEditor/Source/UVEditorTools/Public/UVEditorUXSettings.h
--- a/Engine/Plugins/Editor/UVEditor/Source/UVEditorTools/Public/UVEditorUXSettings.h
+++ b/Engine/Plugins/Editor/UVEditor/Source/UVEditorTools/Public/UVEditorUXSettings.h
@@ -50,6 +50,10 @@ public:
 	// Domain is [-0.5, 0.5] 
 	static FColor MakeTurboColorFromScalar(float Scalar);
 
+	// Wireframe Properties
+	static const float WireframeThickness;
+	static const float BoundaryEdgeThickness;
+
 	// Selection Highlighting Properties
 	static const float SelectionTriangleOpacity;
 	static const FColor SelectionTriangleFillColor;
