UVEditor: Added localization text for missing transform tool categories defined by the detail's customization for the tool.

diff --git a/Engine/Plugins/Editor/UVEditor/Source/UVEditorToolsEditorOnly/Private/DetailsCustomizations/UVTransformToolCustomizations.cpp b/Engine/Plugins/Editor/UVEditor/Source/UVEditorToolsEditorOnly/Private/DetailsCustomizations/UVTransformToolCustomizations.cpp
--- a/Engine/Plugins/Editor/UVEditor/Source/UVEditorToolsEditorOnly/Private/DetailsCustomizations/UVTransformToolCustomizations.cpp
+++ b/Engine/Plugins/Editor/UVEditor/Source/UVEditorToolsEditorOnly/Private/DetailsCustomizations/UVTransformToolCustomizations.cpp
@@ -125,7 +125,8 @@ void FUVEditorUVTransformToolDetails::CustomizeDetails(IDetailLayoutBuilder& Det
 void FUVEditorUVTransformToolDetails::BuildQuickTranslateMenu(IDetailLayoutBuilder& DetailBuilder)
 {
 	const FName QuickTranslateCategoryName = "Quick Translate";
-	IDetailCategoryBuilder& CategoryBuilder = DetailBuilder.EditCategory(QuickTranslateCategoryName);	
+	const FText QuickTranslateCategoryLocName = LOCTEXT("QuickTranslateCategoryName", "Quick Translate");
+	IDetailCategoryBuilder& CategoryBuilder = DetailBuilder.EditCategory(QuickTranslateCategoryName, QuickTranslateCategoryLocName);
 
 	TSharedPtr<IPropertyHandle> QuickTranslateOffsetHandle = DetailBuilder.GetProperty(GET_MEMBER_NAME_CHECKED(UUVEditorUVTransformProperties, QuickTranslateOffset), UUVEditorUVTransformProperties::StaticClass());
 	ensure(QuickTranslateOffsetHandle->IsValidHandle());
@@ -239,7 +240,8 @@ void FUVEditorUVTransformToolDetails::BuildQuickTranslateMenu(IDetailLayoutBuild
 void FUVEditorUVTransformToolDetails::BuildQuickRotateMenu(IDetailLayoutBuilder& DetailBuilder)
 {
 	const FName QuickRotateCategoryName = "Quick Rotate";
-	IDetailCategoryBuilder& CategoryBuilder = DetailBuilder.EditCategory(QuickRotateCategoryName);
+	const FText QuickRotateCategoryLocName = LOCTEXT("QuickRotateCategoryName", "Quick Rotate");
+	IDetailCategoryBuilder& CategoryBuilder = DetailBuilder.EditCategory(QuickRotateCategoryName, QuickRotateCategoryLocName);
 
 	TSharedPtr<IPropertyHandle> QuickRotationOffsetHandle = DetailBuilder.GetProperty(GET_MEMBER_NAME_CHECKED(UUVEditorUVTransformProperties, QuickRotationOffset), UUVEditorUVTransformProperties::StaticClass());
 	ensure(QuickRotationOffsetHandle->IsValidHandle());
@@ -404,8 +406,9 @@ void FUVEditorUVDistributeToolDetails::BuildDistributeModeButtons(IDetailLayoutB
 {
 	using namespace UVEditorTransformDetailsCustomizationLocal;
 
-	const FName QuickRotateCategoryName = "Distribute";
-	IDetailCategoryBuilder& CategoryBuilder = DetailBuilder.EditCategory(QuickRotateCategoryName);
+	const FName DistributeCategoryName = "Distribute";
+	const FText DistributeCategoryLocName = LOCTEXT("DistributeCategoryName", "Distribute");
+	IDetailCategoryBuilder& CategoryBuilder = DetailBuilder.EditCategory(DistributeCategoryName, DistributeCategoryLocName);
 
 	DistributeModeHandle = DetailBuilder.GetProperty(GET_MEMBER_NAME_CHECKED(UUVEditorUVDistributeProperties, DistributeMode), UUVEditorUVDistributeProperties::StaticClass());
 	ensure(DistributeModeHandle->IsValidHandle());
@@ -691,8 +694,9 @@ void FUVEditorUVAlignToolDetails::BuildAlignModeButtons(IDetailLayoutBuilder& De
 {
 	using namespace UVEditorTransformDetailsCustomizationLocal;
 
-	const FName QuickRotateCategoryName = "Align";
-	IDetailCategoryBuilder& CategoryBuilder = DetailBuilder.EditCategory(QuickRotateCategoryName);
+	const FName AlignCategoryName = "Align";
+	const FText AlignCategoryLocName = LOCTEXT("AlignCategoryName", "Align");
+	IDetailCategoryBuilder& CategoryBuilder = DetailBuilder.EditCategory(AlignCategoryName, AlignCategoryLocName);
 
 	AlignDirectionHandle = DetailBuilder.GetProperty(GET_MEMBER_NAME_CHECKED(UUVEditorUVAlignProperties, AlignDirection), UUVEditorUVAlignProperties::StaticClass());
 	ensure(AlignDirectionHandle->IsValidHandle());
