UVEditor: Polishing the new Texel Density tool by adding an icon and adjusting some tooltips behavior.

diff --git a/Engine/Plugins/Editor/UVEditor/Source/UVEditorTools/Private/UVEditorStyle.cpp b/Engine/Plugins/Editor/UVEditor/Source/UVEditorTools/Private/UVEditorStyle.cpp
--- a/Engine/Plugins/Editor/UVEditor/Source/UVEditorTools/Private/UVEditorStyle.cpp
+++ b/Engine/Plugins/Editor/UVEditor/Source/UVEditorTools/Private/UVEditorStyle.cpp
@@ -32,6 +32,7 @@ FUVEditorStyle::FUVEditorStyle()
 	Set("UVEditor.BeginTransformTool",          new IMAGE_BRUSH_SVG("TransformUVs", ToolbarIconSize));
 	Set("UVEditor.BeginAlignTool",              new IMAGE_BRUSH_SVG("AlignLeft", ToolbarIconSize));
 	Set("UVEditor.BeginDistributeTool",         new IMAGE_BRUSH_SVG("DistributeHorizontally", ToolbarIconSize));
+	Set("UVEditor.BeginTexelDensityTool",       new IMAGE_BRUSH_SVG("TexelDensity", ToolbarIconSize));
 
 	// Select tool actions
 	Set("UVEditor.SewAction", new IMAGE_BRUSH_SVG("UVSew", ToolbarIconSize));
diff --git a/Engine/Plugins/Editor/UVEditor/Source/UVEditorTools/Private/UVEditorTexelDensityTool.cpp b/Engine/Plugins/Editor/UVEditor/Source/UVEditorTools/Private/UVEditorTexelDensityTool.cpp
--- a/Engine/Plugins/Editor/UVEditor/Source/UVEditorTools/Private/UVEditorTexelDensityTool.cpp
+++ b/Engine/Plugins/Editor/UVEditor/Source/UVEditorTools/Private/UVEditorTexelDensityTool.cpp
@@ -303,9 +303,7 @@ void UUVEditorTexelDensityTool::Setup()
 
 
 	SetToolDisplayName(LOCTEXT("ToolNameLocal", "Texel Density"));
-	GetToolManager()->DisplayMessage(
-		LOCTEXT("OnStartTool_TexelDensity", "Read and rescale UVs based on texel density values."),
-		EToolMessageLevel::UserNotification);
+	UpdateToolMessage();
 
 	// Analytics
 	InputTargetAnalytics = UVEditorAnalytics::CollectTargetAnalytics(Targets);
@@ -615,6 +613,8 @@ void UUVEditorTexelDensityTool::OnTick(float DeltaTime)
 	{
 		Target->AppliedPreview->Tick(DeltaTime);
 	}
+
+	UpdateToolMessage();
 }
 
 void UUVEditorTexelDensityTool::Render(IToolsContextRenderAPI* RenderAPI)
@@ -652,6 +652,7 @@ void UUVEditorTexelDensityTool::OnPropertyModified(UObject* PropertySet, FProper
 	if (PendingAction == ETexelDensityToolAction::Sampling)
 	{
 		PendingAction = ETexelDensityToolAction::NoAction;
+		UpdateToolMessage();
 	}
 
 	if (PendingAction == ETexelDensityToolAction::NoAction || PendingAction == ETexelDensityToolAction::Processing)
@@ -669,6 +670,8 @@ void UUVEditorTexelDensityTool::RequestAction(ETexelDensityToolAction ActionType
 	{
 		PendingAction = ActionType;
 	}
+
+	UpdateToolMessage();
 }
 
 ETexelDensityToolAction UUVEditorTexelDensityTool::ActiveAction() const
@@ -679,6 +682,7 @@ ETexelDensityToolAction UUVEditorTexelDensityTool::ActiveAction() const
 void UUVEditorTexelDensityTool::PerformBackgroundScalingTask()
 {
 	PendingAction = ETexelDensityToolAction::Processing;
+	UpdateToolMessage();
 
 	RemainingTargetsToProcess = 0;
 	if (UVToolSelectionAPI->HaveSelections())
@@ -697,7 +701,30 @@ void UUVEditorTexelDensityTool::PerformBackgroundScalingTask()
 			Targets[TargetIndex]->AppliedPreview->InvalidateResult();
 		}
 	}
+}
 
+void UUVEditorTexelDensityTool::UpdateToolMessage()
+{
+	switch (PendingAction)
+	{
+	case ETexelDensityToolAction::NoAction:
+		GetToolManager()->DisplayMessage(LOCTEXT("OnStartTool_TexelDensity", "Read and rescale UVs based on texel density values."),
+			EToolMessageLevel::UserNotification);
+		break;
+	case ETexelDensityToolAction::Processing:
+		GetToolManager()->DisplayMessage(LOCTEXT("OnProcessing_TexelDensity", "Applying texel density scaling. Please wait..."),
+			EToolMessageLevel::UserNotification);
+		break;
+	case ETexelDensityToolAction::BeginSamping:
+	case ETexelDensityToolAction::Sampling:
+		GetToolManager()->DisplayMessage(LOCTEXT("OnSampling_TexelDensity", "Click on a triangle in the UV or 3D preview viewport to set the tool's current texel density ratio to be equal to that triangle."),
+			EToolMessageLevel::UserNotification);
+		break;
+	default:
+		ensure(false);
+		GetToolManager()->DisplayMessage(FText(), EToolMessageLevel::UserNotification);
+		break;
+	}
 }
 
 void UUVEditorTexelDensityTool::RecordAnalytics()
diff --git a/Engine/Plugins/Editor/UVEditor/Source/UVEditorTools/Public/UVEditorTexelDensityTool.h b/Engine/Plugins/Editor/UVEditor/Source/UVEditorTools/Public/UVEditorTexelDensityTool.h
--- a/Engine/Plugins/Editor/UVEditor/Source/UVEditorTools/Public/UVEditorTexelDensityTool.h
+++ b/Engine/Plugins/Editor/UVEditor/Source/UVEditorTools/Public/UVEditorTexelDensityTool.h
@@ -185,6 +185,7 @@ private:
 	void ClearHover(bool bClearHoverInfo = true);
 
 	void ApplyClick();
+	void UpdateToolMessage();
 
 	// Used to remember click info to apply on tick
 	int32 ClickedTid = IndexConstants::InvalidID;
