InteractiveToolsFramework: Provide new API for controlling mouse cursor styling during tool operations and prevent excessive hit proxy testing when ITF has mouse capture active.

As a proof of concept, enables the eyedropper cursor during the UV Editor's Texel Density tool's Sampling Mode.

New API is disabled by default, to enable set CVAR Editor.EnableITFCursorOverrideSupport > 1

diff --git a/Engine/Plugins/Editor/ModelingToolsEditorMode/Source/ModelingToolsEditorMode/Private/ModelingToolsEditorMode.cpp b/Engine/Plugins/Editor/ModelingToolsEditorMode/Source/ModelingToolsEditorMode/Private/ModelingToolsEditorMode.cpp
--- a/Engine/Plugins/Editor/ModelingToolsEditorMode/Source/ModelingToolsEditorMode/Private/ModelingToolsEditorMode.cpp
+++ b/Engine/Plugins/Editor/ModelingToolsEditorMode/Source/ModelingToolsEditorMode/Private/ModelingToolsEditorMode.cpp
@@ -1632,6 +1632,13 @@ void UModelingToolsEditorMode::OnToolStarted(UInteractiveToolManager* Manager, U
 		                                            TEXT("ToolName"),
 		                                            GetToolName(*Tool));
 	}
+
+	UContextObjectStore* ContextStore = GetInteractiveToolsContext()->ToolManager->GetContextObjectStore();
+	UToolsContextCursorAPI* ToolsContextCursorAPI = ContextStore->FindContext<UToolsContextCursorAPI>();
+	if (ToolsContextCursorAPI)
+	{
+		ToolsContextCursorAPI->ClearCursorOverride();		
+	}
 }
 
 void UModelingToolsEditorMode::OnToolEnded(UInteractiveToolManager* Manager, UInteractiveTool* Tool)
@@ -1651,6 +1658,13 @@ void UModelingToolsEditorMode::OnToolEnded(UInteractiveToolManager* Manager, UIn
 		                                            TEXT("ToolName"),
 		                                            GetToolName(*Tool));
 	}
+
+	UContextObjectStore* ContextStore = GetInteractiveToolsContext()->ToolManager->GetContextObjectStore();
+	UToolsContextCursorAPI* ToolsContextCursorAPI = ContextStore->FindContext<UToolsContextCursorAPI>();
+	if (ToolsContextCursorAPI)
+	{
+		ToolsContextCursorAPI->ClearCursorOverride();		
+	}
 }
 
 void UModelingToolsEditorMode::BindCommands()
diff --git a/Engine/Plugins/Editor/UVEditor/Source/UVEditor/Private/UVEditorMode.cpp b/Engine/Plugins/Editor/UVEditor/Source/UVEditor/Private/UVEditorMode.cpp
--- a/Engine/Plugins/Editor/UVEditor/Source/UVEditor/Private/UVEditorMode.cpp
+++ b/Engine/Plugins/Editor/UVEditor/Source/UVEditor/Private/UVEditorMode.cpp
@@ -639,6 +639,13 @@ void UUVEditorMode::OnToolStarted(UInteractiveToolManager* Manager, UInteractive
 		GetInteractiveToolsContext()->GetTransactionAPI()->EndUndoTransaction();
 	}
 
+	UContextObjectStore* ContextStore = GetInteractiveToolsContext()->ToolManager->GetContextObjectStore();
+	UToolsContextCursorAPI* ToolsContextCursorAPI = ContextStore->FindContext<UToolsContextCursorAPI>();
+	if (ToolsContextCursorAPI)
+	{
+		ToolsContextCursorAPI->ClearCursorOverride();		
+	}
+
 }
 
 void UUVEditorMode::OnToolEnded(UInteractiveToolManager* Manager, UInteractiveTool* Tool)
@@ -650,6 +657,14 @@ void UUVEditorMode::OnToolEnded(UInteractiveToolManager* Manager, UInteractiveTo
 			Context->OnToolEnded(Tool);
 		}
 	}
+
+	UContextObjectStore* ContextStore = GetInteractiveToolsContext()->ToolManager->GetContextObjectStore();
+	UToolsContextCursorAPI* ToolsContextCursorAPI = ContextStore->FindContext<UToolsContextCursorAPI>();
+	if (ToolsContextCursorAPI)
+	{
+		ToolsContextCursorAPI->ClearCursorOverride();
+	}
+
 }
 
 UObject* UUVEditorMode::GetBackgroundSettingsObject()
@@ -928,6 +943,7 @@ void UUVEditorMode::InitializeAssetEditorContexts(UContextObjectStore& ContextSt
 		ContextStore.AddContextObject(AssetInputsContext);
 	}
 
+	UToolsContextCursorAPI* LivePreviewToolsContextCursorAPI = LivePreviewModeManager.GetInteractiveToolsContext()->ContextObjectStore->FindContext<UToolsContextCursorAPI>();
 	UUVToolLivePreviewAPI* LivePreviewAPI = ContextStore.FindContext<UUVToolLivePreviewAPI>();
 	if (!LivePreviewAPI)
 	{
@@ -947,6 +963,17 @@ void UUVEditorMode::InitializeAssetEditorContexts(UContextObjectStore& ContextSt
 				{
 					LivePreviewViewportClientPtr->FocusViewportOnBox((FBox)BoundingBox, true);
 				}
+			},
+			[LivePreviewToolsContextCursorAPI](const EMouseCursor::Type Cursor, bool bEnableOverride)
+			{
+				if(bEnableOverride)
+				{
+					LivePreviewToolsContextCursorAPI->SetCursorOverride(Cursor);
+				}
+				else
+				{
+					LivePreviewToolsContextCursorAPI->ClearCursorOverride();
+				}
 			}
 			);
 		ContextStore.AddContextObject(LivePreviewAPI);
diff --git a/Engine/Plugins/Editor/UVEditor/Source/UVEditorTools/Private/ContextObjects/UVToolContextObjects.cpp b/Engine/Plugins/Editor/UVEditor/Source/UVEditorTools/Private/ContextObjects/UVToolContextObjects.cpp
--- a/Engine/Plugins/Editor/UVEditor/Source/UVEditorTools/Private/ContextObjects/UVToolContextObjects.cpp
+++ b/Engine/Plugins/Editor/UVEditor/Source/UVEditorTools/Private/ContextObjects/UVToolContextObjects.cpp
@@ -80,18 +80,21 @@ void UUVToolEmitChangeAPI::EmitToolDependentChange(UObject* TargetObject, TUniqu
 
 void UUVToolLivePreviewAPI::Initialize(UWorld* WorldIn, UInputRouter* RouterIn,
 	TUniqueFunction<void(FViewCameraState& CameraStateOut)> GetLivePreviewCameraStateFuncIn,
-	TUniqueFunction<void(const FAxisAlignedBox3d& BoundingBox)> SetLivePreviewCameraToLookAtVolumeFuncIn)
+	TUniqueFunction<void(const FAxisAlignedBox3d& BoundingBox)> SetLivePreviewCameraToLookAtVolumeFuncIn,
+	TUniqueFunction<void(const EMouseCursor::Type Cursor, bool bEnableOverride)> SetCursorOverrideFuncIn)
 {
 	World = WorldIn;
 	InputRouter = RouterIn;
 	GetLivePreviewCameraStateFunc = MoveTemp(GetLivePreviewCameraStateFuncIn);
 	SetLivePreviewCameraToLookAtVolumeFunc = MoveTemp(SetLivePreviewCameraToLookAtVolumeFuncIn);
+	SetCursorOverrideFunc = MoveTemp(SetCursorOverrideFuncIn);
 }
 
 void UUVToolLivePreviewAPI::OnToolEnded(UInteractiveTool* DeadTool)
 {
 	OnDrawHUD.RemoveAll(DeadTool);
 	OnRender.RemoveAll(DeadTool);
+	ClearCursorOverride();
 }
 
 int32 FUDIMBlock::BlockU() const
diff --git a/Engine/Plugins/Editor/UVEditor/Source/UVEditorTools/Private/UVEditorTexelDensityTool.cpp b/Engine/Plugins/Editor/UVEditor/Source/UVEditorTools/Private/UVEditorTexelDensityTool.cpp
--- a/Engine/Plugins/Editor/UVEditor/Source/UVEditorTools/Private/UVEditorTexelDensityTool.cpp
+++ b/Engine/Plugins/Editor/UVEditor/Source/UVEditorTools/Private/UVEditorTexelDensityTool.cpp
@@ -575,8 +575,16 @@ void UUVEditorTexelDensityTool::Shutdown(EToolShutdownType ShutdownType)
 
 void UUVEditorTexelDensityTool::OnTick(float DeltaTime)
 {
+	UToolsContextCursorAPI* ToolsContextCursorAPI = GetToolManager()->GetContextObjectStore()->FindContext<UToolsContextCursorAPI>();
+
 	if (PendingAction == ETexelDensityToolAction::BeginSamping)
 	{
+		if (ToolsContextCursorAPI)
+		{
+			ToolsContextCursorAPI->SetCursorOverride(EMouseCursor::EyeDropper);			
+		}
+		LivePreviewAPI->SetCursorOverride(EMouseCursor::EyeDropper);
+
 		for (const TObjectPtr<UUVEditorToolMeshInput>& Target : Targets)
 		{
 			Target->UpdatePreviewsFromCanonical();
@@ -595,6 +603,12 @@ void UUVEditorTexelDensityTool::OnTick(float DeltaTime)
 			ApplyClick();
 			ClickedTid = IndexConstants::InvalidID;
 
+			if (ToolsContextCursorAPI)
+			{
+				ToolsContextCursorAPI->ClearCursorOverride();
+			}
+			LivePreviewAPI->ClearCursorOverride();
+
 			PendingAction = ETexelDensityToolAction::NoAction;
 			PerformBackgroundScalingTask();
 		}
diff --git a/Engine/Plugins/Editor/UVEditor/Source/UVEditorTools/Public/ContextObjects/UVToolContextObjects.h b/Engine/Plugins/Editor/UVEditor/Source/UVEditorTools/Public/ContextObjects/UVToolContextObjects.h
--- a/Engine/Plugins/Editor/UVEditor/Source/UVEditorTools/Public/ContextObjects/UVToolContextObjects.h
+++ b/Engine/Plugins/Editor/UVEditor/Source/UVEditorTools/Public/ContextObjects/UVToolContextObjects.h
@@ -108,7 +108,8 @@ public:
 
 	void Initialize(UWorld* WorldIn, UInputRouter* RouterIn,
 		TUniqueFunction<void(FViewCameraState& CameraStateOut)> GetLivePreviewCameraStateFuncIn,
-		TUniqueFunction<void(const UE::Geometry::FAxisAlignedBox3d& BoundingBox)> SetLivePreviewCameraToLookAtVolumeFuncIn);
+		TUniqueFunction<void(const UE::Geometry::FAxisAlignedBox3d& BoundingBox)> SetLivePreviewCameraToLookAtVolumeFuncIn,
+		TUniqueFunction<void(const EMouseCursor::Type Cursor, bool bEnableOverride)> SetCursorOverrideFuncIn);
 
 	UWorld* GetLivePreviewWorld() { return World.Get(); }
 	UInputRouter* GetLivePreviewInputRouter() { return InputRouter.Get(); }
@@ -128,6 +129,22 @@ public:
 		}
 	}
 
+	void SetCursorOverride(const EMouseCursor::Type Cursor)
+	{
+		if (SetCursorOverrideFunc)
+		{
+			SetCursorOverrideFunc(Cursor, true);
+		}
+	}
+
+	void ClearCursorOverride()
+	{
+		if (SetCursorOverrideFunc)
+		{
+			SetCursorOverrideFunc(EMouseCursor::Default, false);
+		}
+	}
+
 	virtual void OnToolEnded(UInteractiveTool* DeadTool) override;
 
 	/**
@@ -153,6 +170,7 @@ protected:
 
 	TUniqueFunction<void(FViewCameraState& CameraStateOut)> GetLivePreviewCameraStateFunc;
 	TUniqueFunction<void(const UE::Geometry::FAxisAlignedBox3d& BoundingBox)> SetLivePreviewCameraToLookAtVolumeFunc;
+	TUniqueFunction<void(const EMouseCursor::Type Cursor, bool bEnableOverride)> SetCursorOverrideFunc;
 };
 
 USTRUCT()
diff --git a/Engine/Source/Editor/Experimental/EditorInteractiveToolsFramework/Private/EdModeInteractiveToolsContext.cpp b/Engine/Source/Editor/Experimental/EditorInteractiveToolsFramework/Private/EdModeInteractiveToolsContext.cpp
--- a/Engine/Source/Editor/Experimental/EditorInteractiveToolsFramework/Private/EdModeInteractiveToolsContext.cpp
+++ b/Engine/Source/Editor/Experimental/EditorInteractiveToolsFramework/Private/EdModeInteractiveToolsContext.cpp
@@ -40,9 +40,6 @@
 
 //#define ENABLE_DEBUG_PRINTING
 
-
-
-
 class FEdModeToolsContextQueriesImpl : public IToolsContextQueriesAPI
 {
 public:
@@ -315,6 +312,13 @@ void UEditorInteractiveToolsContext::Initialize(IToolsContextQueriesAPI* Queries
 
 	// This gets set up in UInteractiveToolsContext::Initialize;
 	GizmoViewContext = ToolManager->GetContextObjectStore()->FindContext<UGizmoViewContext>();
+
+	UToolsContextCursorAPI* ToolsContextCursorAPI = ToolManager->GetContextObjectStore()->FindContext<UToolsContextCursorAPI>();
+	if (!ToolsContextCursorAPI)
+	{
+		ToolsContextCursorAPI = NewObject< UToolsContextCursorAPI >();
+		ToolManager->GetContextObjectStore()->AddContextObject(ToolsContextCursorAPI);
+	}
 }
 
 void UEditorInteractiveToolsContext::Shutdown()
@@ -1140,6 +1144,25 @@ bool UModeManagerInteractiveToolsContext::EndTracking(FEditorViewportClient* InV
 	return false;
 }
 
+bool UModeManagerInteractiveToolsContext::GetCursor(EMouseCursor::Type& OutCursor) const
+{	
+	UToolsContextCursorAPI* ToolsContextCursorAPI = ToolManager->GetContextObjectStore()->FindContext<UToolsContextCursorAPI>();
+
+	if (ToolsContextCursorAPI && ToolsContextCursorAPI->IsCursorOverridden())
+	{
+		OutCursor = ToolsContextCursorAPI->GetCurrentCursorOverride();
+		return true;
+	}
+
+	if (bIsTrackingMouse)
+	{
+		OutCursor = EMouseCursor::SlashedCircle;				
+		return true;
+	}
+
+	return false;
+}
+
 FRay UModeManagerInteractiveToolsContext::GetLastWorldRay() const
 {
 	return CurrentMouseState.Mouse.WorldRay;
diff --git a/Engine/Source/Editor/Experimental/EditorInteractiveToolsFramework/Public/EdModeInteractiveToolsContext.h b/Engine/Source/Editor/Experimental/EditorInteractiveToolsFramework/Public/EdModeInteractiveToolsContext.h
--- a/Engine/Source/Editor/Experimental/EditorInteractiveToolsFramework/Public/EdModeInteractiveToolsContext.h
+++ b/Engine/Source/Editor/Experimental/EditorInteractiveToolsFramework/Public/EdModeInteractiveToolsContext.h
@@ -271,6 +271,9 @@ public:
 	bool CapturedMouseMove(FEditorViewportClient* InViewportClient, FViewport* InViewport, int32 InMouseX, int32 InMouseY);
 	bool EndTracking(FEditorViewportClient* InViewportClient, FViewport* InViewport);
 
+	/** @return True if the context has overriden the cursor style, false if not. */
+	bool GetCursor(EMouseCursor::Type& OutCursor) const;
+
 	/** @return Ray into 3D scene at last mouse event */
 	virtual FRay GetLastWorldRay() const override;
 
diff --git a/Engine/Source/Editor/UnrealEd/Private/EditorModeManager.cpp b/Engine/Source/Editor/UnrealEd/Private/EditorModeManager.cpp
--- a/Engine/Source/Editor/UnrealEd/Private/EditorModeManager.cpp
+++ b/Engine/Source/Editor/UnrealEd/Private/EditorModeManager.cpp
@@ -54,6 +54,11 @@
 	The master class that handles tracking of the current mode.
 ------------------------------------------------------------------------------*/
 
+TAutoConsoleVariable<int32> FEditorModeTools::CVarEnableITFCursorOverrideSupport(
+	TEXT("Editor.EnableITFCursorOverrideSupport"),
+	0,
+	TEXT("Enable support for the EditorModeManager to query the InteractiveToolsContext for tool customized cursor overrides."));
+
 FEditorModeTools::FEditorModeTools()
 	: PivotShown(false)
 	, Snapping(false)
@@ -1707,9 +1712,18 @@ bool FEditorModeTools::DisallowMouseDeltaTracking() const
 bool FEditorModeTools::GetCursor(EMouseCursor::Type& OutCursor) const
 {
 	bool bHandled = false;
-	for (const UEdMode* Mode : ActiveScriptableModes)
+
+	if (CVarEnableITFCursorOverrideSupport.GetValueOnGameThread() > 0)
 	{
-		bHandled |= Mode->GetCursor(OutCursor);
+		bHandled = InteractiveToolsContext->GetCursor(OutCursor);
+	}
+
+	if (!bHandled)
+	{
+		for (const UEdMode* Mode : ActiveScriptableModes)
+		{
+			bHandled |= Mode->GetCursor(OutCursor);
+		}
 	}
 	return bHandled;
 }
diff --git a/Engine/Source/Editor/UnrealEd/Public/EditorModeManager.h b/Engine/Source/Editor/UnrealEd/Public/EditorModeManager.h
--- a/Engine/Source/Editor/UnrealEd/Public/EditorModeManager.h
+++ b/Engine/Source/Editor/UnrealEd/Public/EditorModeManager.h
@@ -646,4 +646,6 @@ private:
 	FEditorViewportClient* FocusedViewportClient = nullptr;
 
 	TMap<FName, FTypedElementSelectionSetState> StoredSelectionSets;
+
+	static TAutoConsoleVariable<int32> CVarEnableITFCursorOverrideSupport;
 };
diff --git a/Engine/Source/Runtime/InteractiveToolsFramework/Public/ToolContextInterfaces.h b/Engine/Source/Runtime/InteractiveToolsFramework/Public/ToolContextInterfaces.h
--- a/Engine/Source/Runtime/InteractiveToolsFramework/Public/ToolContextInterfaces.h
+++ b/Engine/Source/Runtime/InteractiveToolsFramework/Public/ToolContextInterfaces.h
@@ -17,6 +17,7 @@
 #include "UObject/WeakObjectPtr.h"
 #include "UObject/WeakObjectPtr.h"
 #include "UObject/WeakObjectPtrTemplates.h"
+#include "GenericPlatform/ICursor.h"
 
 #include "ToolContextInterfaces.generated.h"
 
@@ -173,6 +174,52 @@ struct FToolContextSnappingConfiguration
 };
 
 
+/**
+* Users of the Tools Framework can use  UToolsContextCursorAPI to
+* access to setting and retrieving cursor overrides for tool activities.
+*/
+UCLASS()
+class INTERACTIVETOOLSFRAMEWORK_API UToolsContextCursorAPI : public UObject
+{
+	GENERATED_BODY()
+
+private:
+	EMouseCursor::Type CursorOverrideShape;
+	bool bIsCursorOverridden;
+
+public:
+
+	UToolsContextCursorAPI()
+	{
+		CursorOverrideShape = EMouseCursor::Default;
+		bIsCursorOverridden = false;
+	}
+
+	virtual ~UToolsContextCursorAPI() {}
+
+	virtual EMouseCursor::Type GetCurrentCursorOverride() const 
+	{
+		return CursorOverrideShape;
+	}
+
+	virtual bool IsCursorOverridden() const
+	{
+		return bIsCursorOverridden;
+	}
+
+	virtual void SetCursorOverride(EMouseCursor::Type CursorOverride)
+	{
+		bIsCursorOverridden = true;
+		CursorOverrideShape = CursorOverride;
+	}
+
+	virtual void ClearCursorOverride() 
+	{
+		bIsCursorOverridden = false;
+		CursorOverrideShape = EMouseCursor::Default;
+	}
+};
+
 /**
  * Users of the Tools Framework need to implement IToolsContextQueriesAPI to provide
  * access to scene state information like the current UWorld, active USelections, etc.
