Modeling Tools: Fixes bug in Seam Tool visalization when using multiple meshes or rotation in UV editor

diff --git a/Engine/Plugins/Editor/UVEditor/Source/UVEditorTools/Private/UVEditorSeamTool.cpp b/Engine/Plugins/Editor/UVEditor/Source/UVEditorTools/Private/UVEditorSeamTool.cpp
--- a/Engine/Plugins/Editor/UVEditor/Source/UVEditorTools/Private/UVEditorSeamTool.cpp
+++ b/Engine/Plugins/Editor/UVEditor/Source/UVEditorTools/Private/UVEditorSeamTool.cpp
@@ -123,10 +123,10 @@ namespace UVEditorSeamToolLocals
 	{
 		TRACE_CPUPROFILER_EVENT_SCOPE(AddDisplayedPoints);
 
-		FTransform AppliedTransform = InputObject->AppliedPreview->PreviewMesh->GetTransform();
+		const FTransform AppliedTransform = InputObject->AppliedPreview->PreviewMesh->GetTransform();
 		AppliedPointSet->AddPoint(FRenderablePoint(
-			InputObject->AppliedCanonical->GetVertex(AppliedVid), Color, FUVEditorUXSettings::ToolPointSize, DepthBias));
-		for (int32 UnwrapVid : UnwrapVids)
+			AppliedTransform.TransformPosition(InputObject->AppliedCanonical->GetVertex(AppliedVid)), Color, FUVEditorUXSettings::ToolPointSize, DepthBias));
+		for (const int32 UnwrapVid : UnwrapVids)
 		{
 			UnwrapPointSet->AddPoint(FRenderablePoint(
 				InputObject->UnwrapCanonical->GetVertex(UnwrapVid), Color, FUVEditorUXSettings::ToolPointSize, DepthBias));
@@ -482,8 +482,7 @@ void UUVEditorSeamTool::Setup()
 	UnwrapGeometry->CreateInWorld(Targets[0]->UnwrapPreview->GetWorld(), FTransform::Identity);
 	
 	LivePreviewGeometry = NewObject<UPreviewGeometry>();
-	LivePreviewGeometry->CreateInWorld(Targets[0]->AppliedPreview->GetWorld(),
-		Targets[0]->AppliedPreview->PreviewMesh->GetTransform());
+	LivePreviewGeometry->CreateInWorld(Targets[0]->AppliedPreview->GetWorld(), FTransform::Identity);
 
 	// These visualize the locked-in portion of the current seam
 	UnwrapGeometry->AddPointSet(LockedPointSetID);
