Modeling Tools: Adds UX customization settings in the UV editor display menu

diff --git a/Engine/Plugins/Editor/UVEditor/Source/UVEditor/Private/UVEditorMode.cpp b/Engine/Plugins/Editor/UVEditor/Source/UVEditor/Private/UVEditorMode.cpp
--- a/Engine/Plugins/Editor/UVEditor/Source/UVEditor/Private/UVEditorMode.cpp
+++ b/Engine/Plugins/Editor/UVEditor/Source/UVEditor/Private/UVEditorMode.cpp
@@ -54,6 +54,7 @@
 #include "UVEditorLogging.h"
 #include "UObject/ObjectSaveContext.h"
 #include "Materials/Material.h"
+#include "UVEditorUXPropertySets.h"
 
 #include UE_INLINE_GENERATED_CPP_BY_NAME(UVEditorMode)
 
@@ -441,6 +442,82 @@ void UUVEditorMode::Enter()
 		});
 
 	PropertyObjectsToTick.Add(UVEditorGridProperties);
+	
+	UVEditorUnwrappedUXProperties = NewObject<UUVEditorUnwrappedUXProperties>(this);
+	
+	UVEditorUnwrappedUXProperties->BoundaryLineColors.Reserve(ToolInputObjects.Num());
+	// initialize display menu properties
+	for (int ObjIndex = 0; ObjIndex < ToolInputObjects.Num(); ObjIndex++)
+	{
+		UVEditorUnwrappedUXProperties->BoundaryLineColors.Add(ToolInputObjects[ObjIndex]->WireframeDisplay->Settings->BoundaryEdgeColor);
+	}
+	UVEditorUnwrappedUXProperties->BoundaryLineThickness = FUVEditorUXSettings::BoundaryEdgeThickness;
+	UVEditorUnwrappedUXProperties->WireframeThickness = FUVEditorUXSettings::WireframeThickness;
+
+	UVEditorUnwrappedUXProperties->WatchProperty(UVEditorUnwrappedUXProperties->BoundaryLineColors,
+		[this](const TArray<FColor>& EdgeColors)
+		{
+			for (int32 AssetID = 0; AssetID < ToolInputObjects.Num(); ++AssetID)
+			{
+				ToolInputObjects[AssetID]->WireframeDisplay->Settings->BoundaryEdgeColor = EdgeColors[AssetID];
+			}
+		});
+	UVEditorUnwrappedUXProperties->WatchProperty(UVEditorUnwrappedUXProperties->BoundaryLineThickness,
+		[this](const float Thickness)
+		{
+			for (int32 AssetID = 0; AssetID < ToolInputObjects.Num(); ++AssetID)
+			{
+				ToolInputObjects[AssetID]->WireframeDisplay->WireframeComponent->BoundaryEdgeThickness = Thickness;
+				ToolInputObjects[AssetID]->WireframeDisplay->WireframeComponent->UpdateWireframe();
+			}
+		});
+	UVEditorUnwrappedUXProperties->WatchProperty(UVEditorUnwrappedUXProperties->WireframeThickness,
+		[this](const float WireframeThickness)
+		{
+			for (int32 AssetID = 0; AssetID < ToolInputObjects.Num(); ++AssetID)
+			{
+				ToolInputObjects[AssetID]->WireframeDisplay->WireframeComponent->WireframeThickness = WireframeThickness;
+				ToolInputObjects[AssetID]->WireframeDisplay->WireframeComponent->UpdateWireframe();
+			}
+		});
+	PropertyObjectsToTick.Add(UVEditorUnwrappedUXProperties);
+
+	UVEditorLivePreviewUXProperties = NewObject<UUVEditorLivePreviewUXProperties>(this);
+
+	// initialize display menu properties
+	UVEditorLivePreviewUXProperties->SelectionColor = FUVEditorUXSettings::SelectionTriangleWireframeColor;
+	UVEditorLivePreviewUXProperties->SelectionLineThickness = FUVEditorUXSettings::LivePreviewHighlightThickness;
+	UVEditorLivePreviewUXProperties->SelectionPointSize = FUVEditorUXSettings::LivePreviewHighlightPointSize;
+	
+	UVEditorLivePreviewUXProperties->WatchProperty(UVEditorLivePreviewUXProperties->SelectionColor,
+		[this](const FColor InSelectionLineColor)
+		{
+			UUVToolSelectionAPI::FLivePreviewSelectionUXSettings Settings = UUVToolSelectionAPI::FLivePreviewSelectionUXSettings();
+			Settings.SelectionColor = InSelectionLineColor;
+
+			SelectionAPI->SetLivePreviewSelectionUXSettings(Settings);
+			SelectionAPI->RebuildAppliedPreviewHighlight();
+		});
+	UVEditorLivePreviewUXProperties->WatchProperty(UVEditorLivePreviewUXProperties->SelectionLineThickness,
+		[this](const float InLineThickness)
+		{
+			UUVToolSelectionAPI::FLivePreviewSelectionUXSettings Settings = UUVToolSelectionAPI::FLivePreviewSelectionUXSettings();
+			Settings.LineThickness = InLineThickness;
+			
+			SelectionAPI->SetLivePreviewSelectionUXSettings(Settings);
+			SelectionAPI->RebuildAppliedPreviewHighlight();
+		});
+	UVEditorLivePreviewUXProperties->WatchProperty(UVEditorLivePreviewUXProperties->SelectionPointSize,
+		[this](const float InPointSize)
+		{
+			UUVToolSelectionAPI::FLivePreviewSelectionUXSettings Settings = UUVToolSelectionAPI::FLivePreviewSelectionUXSettings();
+			Settings.PointSize = InPointSize;
+			
+			SelectionAPI->SetLivePreviewSelectionUXSettings(Settings);
+			SelectionAPI->RebuildAppliedPreviewHighlight();
+		});
+
+	PropertyObjectsToTick.Add(UVEditorLivePreviewUXProperties);
 
 	UVEditorUDIMProperties = NewObject< UUVEditorUDIMProperties >(this);
 	UVEditorUDIMProperties->Initialize(this);
@@ -694,6 +771,16 @@ UObject* UUVEditorMode::GetGridSettingsObject()
 	return nullptr;
 }
 
+UObject* UUVEditorMode::GetUnwrappedUXSettingsObject() const
+{
+	return UVEditorUnwrappedUXProperties.Get();
+}
+
+UObject* UUVEditorMode::GetLivePreviewUXSettingsObject() const
+{
+	return UVEditorLivePreviewUXProperties.Get();
+}
+
 UObject* UUVEditorMode::GetUDIMSettingsObject()
 {
 	if (UVEditorUDIMProperties)
diff --git a/Engine/Plugins/Editor/UVEditor/Source/UVEditor/Private/UVEditorModeToolkit.cpp b/Engine/Plugins/Editor/UVEditor/Source/UVEditor/Private/UVEditorModeToolkit.cpp
--- a/Engine/Plugins/Editor/UVEditor/Source/UVEditor/Private/UVEditorModeToolkit.cpp
+++ b/Engine/Plugins/Editor/UVEditor/Source/UVEditor/Private/UVEditorModeToolkit.cpp
@@ -311,6 +311,18 @@ TSharedRef<SWidget> FUVEditorModeToolkit::CreateGridSettingsWidget()
 	return CreateDisplaySettingsWidget(Mode->GetGridSettingsObject());
 }
 
+TSharedRef<SWidget> FUVEditorModeToolkit::CreateUnwrappedUXSettingsWidget()
+{
+	UUVEditorMode* Mode = Cast<UUVEditorMode>(GetScriptableEditorMode());
+	return CreateDisplaySettingsWidget(Mode->GetUnwrappedUXSettingsObject());
+}
+
+TSharedRef<SWidget> FUVEditorModeToolkit::CreateLivePreviewUXSettingsWidget()
+{
+	UUVEditorMode* Mode = Cast<UUVEditorMode>(GetScriptableEditorMode());
+	return CreateDisplaySettingsWidget(Mode->GetLivePreviewUXSettingsObject());
+}
+
 TSharedRef<SWidget> FUVEditorModeToolkit::CreateUDIMSettingsWidget()
 {
 	UUVEditorMode* Mode = Cast<UUVEditorMode>(GetScriptableEditorMode());
diff --git a/Engine/Plugins/Editor/UVEditor/Source/UVEditor/Private/UVEditorToolkit.cpp b/Engine/Plugins/Editor/UVEditor/Source/UVEditor/Private/UVEditorToolkit.cpp
--- a/Engine/Plugins/Editor/UVEditor/Source/UVEditor/Private/UVEditorToolkit.cpp
+++ b/Engine/Plugins/Editor/UVEditor/Source/UVEditor/Private/UVEditorToolkit.cpp
@@ -569,6 +569,28 @@ void FUVEditorToolkit::PostInitAssetEditor()
 				]
 				];
 
+			Container->AddSlot()
+				.AutoHeight()
+				.Padding(FMargin(0.f, 0.f, 8.f, 0.f))
+				[
+					SNew(SBox)
+					.MinDesiredWidth(500)
+				[
+					UVModeToolkit->CreateUnwrappedUXSettingsWidget()
+				]
+				];
+
+			Container->AddSlot()
+				.AutoHeight()
+				.Padding(FMargin(0.f, 0.f, 8.f, 0.f))
+				[
+					SNew(SBox)
+					.MinDesiredWidth(500)
+				[
+					UVModeToolkit->CreateLivePreviewUXSettingsWidget()
+				]
+				];
+
 			Container->AddSlot()
 				.AutoHeight()
 				.Padding(FMargin(0.f, 0.f, 8.f, 0.f))
diff --git a/Engine/Plugins/Editor/UVEditor/Source/UVEditor/Private/UVEditorUXPropertySets.h b/Engine/Plugins/Editor/UVEditor/Source/UVEditor/Private/UVEditorUXPropertySets.h
new file mode 100644
--- /dev/null
+++ b/Engine/Plugins/Editor/UVEditor/Source/UVEditor/Private/UVEditorUXPropertySets.h
@@ -0,0 +1,44 @@
+﻿// Copyright Epic Games, Inc. All Rights Reserved.
+#pragma once
+
+#include "InteractiveTool.h"
+#include "UVEditorMode.h"
+
+#include "UVEditorUXPropertySets.generated.h"
+
+/**
+* Visualization settings for the UUVEditorMode's Unwrapped Visualization
+ */
+UCLASS()
+class UVEDITOR_API UUVEditorUnwrappedUXProperties : public UInteractiveToolPropertySet
+{
+	GENERATED_BODY()
+public:
+	UPROPERTY(EditAnywhere, Category="Unwrapped Visualization", meta = (DisplayName = "Boundary Line Color"))
+	TArray<FColor> BoundaryLineColors;
+
+	UPROPERTY(EditAnywhere, Category="Unwrapped Visualization", meta = (DisplayName = "Boundary Line Thickness"))
+	float BoundaryLineThickness;
+
+	// TODO: maybe opacity is the setting to change instead of/in addition to thickness? Could also have something to recognize line density and/or window size and adjust accordingly
+	UPROPERTY(EditAnywhere, Category="Unwrapped Visualization", meta = (DisplayName = "Wireframe Thickness"))
+	float WireframeThickness;
+};
+
+/**
+ * Visualization settings for the UUVEditorMode's properties in the Live Preview (3d) Viewport
+ */
+UCLASS()
+class UVEDITOR_API UUVEditorLivePreviewUXProperties : public UInteractiveToolPropertySet
+{
+	GENERATED_BODY()
+public:
+	UPROPERTY(EditAnywhere, Category="Live Preview Visualization", meta = (DisplayName = "Selection Line Color"))
+	FColor SelectionColor;
+
+	UPROPERTY(EditAnywhere, Category="Live Preview Visualization", meta = (DisplayName = "Selection Line Thickness"))
+	float SelectionLineThickness;
+
+	UPROPERTY(EditAnywhere, Category="Live Preview Visualization", meta = (DisplayName = "Selection Point Size"))
+	float SelectionPointSize;
+};
\ No newline at end of file
diff --git a/Engine/Plugins/Editor/UVEditor/Source/UVEditor/Public/UVEditorMode.h b/Engine/Plugins/Editor/UVEditor/Source/UVEditor/Public/UVEditorMode.h
--- a/Engine/Plugins/Editor/UVEditor/Source/UVEditor/Public/UVEditorMode.h
+++ b/Engine/Plugins/Editor/UVEditor/Source/UVEditor/Public/UVEditorMode.h
@@ -33,6 +33,8 @@ class UUVToolViewportButtonsAPI;
 class UUVTool2DViewportAPI;
 class UUVEditorMode;
 class UWorld;
+class UUVEditorUnwrappedUXProperties;
+class UUVEditorLivePreviewUXProperties;
 
 /**
  * Visualization settings for the UUVEditorMode's Grid
@@ -220,6 +222,12 @@ public:
 	/** @return A settings object suitable for display in a details panel to control the grid. */
 	UObject* GetGridSettingsObject();
 
+	/** @return A settings object suitable for display in a details panel to control the unwrapped viewport visualization. */
+	UObject* GetUnwrappedUXSettingsObject() const;
+
+	/** @return A settings object suitable for display in a details panel to control the visualization in the live preview (3d) viewport. */
+	UObject* GetLivePreviewUXSettingsObject() const;
+
 	/** @return A settings object suitable for display in a details panel to control UDIM configuration. */
 	UObject* GetUDIMSettingsObject();
 
@@ -387,5 +395,14 @@ protected:
 	// Holds references to Save callbacks to handle logic when the autosave triggers and shuts down active tools. 
 	// We need to recover from this, so we restart the select tool after the save is over.
 	FDelegateHandle PostSaveWorldDelegateHandle;
+private:
+	// Hold a settings object to configure the unwrapped visualization
+	UPROPERTY()
+	TObjectPtr<UUVEditorUnwrappedUXProperties> UVEditorUnwrappedUXProperties = nullptr;
+
+	// Hold a settings object to configure the visualization in the live preview (3d) viewport
+	UPROPERTY()
+	TObjectPtr<UUVEditorLivePreviewUXProperties> UVEditorLivePreviewUXProperties = nullptr;
+
 };
 
diff --git a/Engine/Plugins/Editor/UVEditor/Source/UVEditor/Public/UVEditorModeToolkit.h b/Engine/Plugins/Editor/UVEditor/Source/UVEditor/Public/UVEditorModeToolkit.h
--- a/Engine/Plugins/Editor/UVEditor/Source/UVEditor/Public/UVEditorModeToolkit.h
+++ b/Engine/Plugins/Editor/UVEditor/Source/UVEditor/Public/UVEditorModeToolkit.h
@@ -34,6 +34,12 @@ public:
 
 	/** Creates a widget where the grid display can be changed. */
 	virtual TSharedRef<SWidget> CreateGridSettingsWidget();
+
+	/** Creates a widget where the visualization in the 2d/unwrapped view can be changed. */
+	virtual TSharedRef<SWidget> CreateUnwrappedUXSettingsWidget();
+
+	/** Creates a widget where the visualization properties in the 3d viewport can be changed. */
+	virtual TSharedRef<SWidget> CreateLivePreviewUXSettingsWidget();
 	
 	/** Creates a widget where the visible UDIMs can be changed. */
 	virtual TSharedRef<SWidget> CreateUDIMSettingsWidget();
diff --git a/Engine/Plugins/Editor/UVEditor/Source/UVEditorTools/Private/Drawing/BasicLineSetComponent.cpp b/Engine/Plugins/Editor/UVEditor/Source/UVEditorTools/Private/Drawing/BasicLineSetComponent.cpp
--- a/Engine/Plugins/Editor/UVEditor/Source/UVEditorTools/Private/Drawing/BasicLineSetComponent.cpp
+++ b/Engine/Plugins/Editor/UVEditor/Source/UVEditorTools/Private/Drawing/BasicLineSetComponent.cpp
@@ -243,6 +243,17 @@ void UBasicLineSetComponentBase::SetLineSetParameters(FColor InColor, float InSi
 	bBoundsDirty = true;
 }
 
+void UBasicLineSetComponentBase::SetColor(const FColor InColor)
+{
+	Color = InColor;
+}
+
+void UBasicLineSetComponentBase::SetLineThickness(const float InSize)
+{
+	Size = InSize;
+	bBoundsDirty = true;
+}
+
 
 FPrimitiveSceneProxy* UBasic2DLineSetComponent::CreateSceneProxy()
 {
diff --git a/Engine/Plugins/Editor/UVEditor/Source/UVEditorTools/Private/Drawing/BasicPointSetComponent.cpp b/Engine/Plugins/Editor/UVEditor/Source/UVEditorTools/Private/Drawing/BasicPointSetComponent.cpp
--- a/Engine/Plugins/Editor/UVEditor/Source/UVEditorTools/Private/Drawing/BasicPointSetComponent.cpp
+++ b/Engine/Plugins/Editor/UVEditor/Source/UVEditorTools/Private/Drawing/BasicPointSetComponent.cpp
@@ -234,6 +234,17 @@ void UBasicPointSetComponentBase::SetPointSetParameters(FColor InColor, float In
 	bBoundsDirty = true;
 }
 
+void UBasicPointSetComponentBase::SetColor(const FColor InColor)
+{
+	Color = InColor;
+}
+
+void UBasicPointSetComponentBase::SetPointSize(const float InSize)
+{
+	Size = InSize;
+	bBoundsDirty = true;
+}
+
 
 FPrimitiveSceneProxy* UBasic2DPointSetComponent::CreateSceneProxy()
 {
diff --git a/Engine/Plugins/Editor/UVEditor/Source/UVEditorTools/Private/Selection/UVToolSelectionAPI.cpp b/Engine/Plugins/Editor/UVEditor/Source/UVEditorTools/Private/Selection/UVToolSelectionAPI.cpp
--- a/Engine/Plugins/Editor/UVEditor/Source/UVEditorTools/Private/Selection/UVToolSelectionAPI.cpp
+++ b/Engine/Plugins/Editor/UVEditor/Source/UVEditorTools/Private/Selection/UVToolSelectionAPI.cpp
@@ -147,6 +147,24 @@ void UUVToolSelectionAPI::OnToolEnded(UInteractiveTool* DeadTool)
 	
 }
 
+void UUVToolSelectionAPI::SetLivePreviewSelectionUXSettings(const FLivePreviewSelectionUXSettings& Settings) const
+{
+	if (Settings.SelectionColor)
+	{
+		HighlightMechanic->SetColor(Settings.SelectionColor.GetValue());
+	}
+
+	if (Settings.LineThickness)
+	{
+		HighlightMechanic->SetLineThickness(Settings.LineThickness.GetValue());
+	}
+
+	if (Settings.PointSize)
+	{
+		HighlightMechanic->SetPointSize(Settings.PointSize.GetValue());
+	}
+}
+
 void UUVToolSelectionAPI::SetSelections(const TArray<FUVToolSelection>& SelectionsIn, bool bBroadcast, bool bEmitChange)
 {
 	using namespace UVToolSelectionAPILocals;
diff --git a/Engine/Plugins/Editor/UVEditor/Source/UVEditorTools/Private/Selection/UVToolSelectionHighlightMechanic.cpp b/Engine/Plugins/Editor/UVEditor/Source/UVEditorTools/Private/Selection/UVToolSelectionHighlightMechanic.cpp
--- a/Engine/Plugins/Editor/UVEditor/Source/UVEditorTools/Private/Selection/UVToolSelectionHighlightMechanic.cpp
+++ b/Engine/Plugins/Editor/UVEditor/Source/UVEditorTools/Private/Selection/UVToolSelectionHighlightMechanic.cpp
@@ -609,4 +609,20 @@ void UUVToolSelectionHighlightMechanic::SetEnablePairedEdgeHighlights(bool bEnab
 	SewEdgeUnselectedPairingLineSet->SetVisibility(bEnable);
 }
 
+void UUVToolSelectionHighlightMechanic::SetColor(const FColor Color) const
+{
+	LivePreviewLineSet->SetColor(Color);
+	LivePreviewPointSet->SetColor(Color);
+}
+
+void UUVToolSelectionHighlightMechanic::SetLineThickness(const float LineSize) const
+{
+	LivePreviewLineSet->SetLineThickness(LineSize);
+}
+
+void UUVToolSelectionHighlightMechanic::SetPointSize(const float PointSize) const
+{
+	LivePreviewPointSet->SetPointSize(PointSize);
+}
+
 #undef LOCTEXT_NAMESPACE
diff --git a/Engine/Plugins/Editor/UVEditor/Source/UVEditorTools/Public/Drawing/BasicLineSetComponent.h b/Engine/Plugins/Editor/UVEditor/Source/UVEditorTools/Public/Drawing/BasicLineSetComponent.h
--- a/Engine/Plugins/Editor/UVEditor/Source/UVEditorTools/Public/Drawing/BasicLineSetComponent.h
+++ b/Engine/Plugins/Editor/UVEditor/Source/UVEditorTools/Public/Drawing/BasicLineSetComponent.h
@@ -44,6 +44,12 @@ public:
 	/** Set per line material parameters that are uniform for all lines. */
 	void SetLineSetParameters(FColor InColor, float InSize, float InDepthBias);
 
+	/** Set per line material Line Color parameter that is uniform for all lines */
+	void SetColor(FColor InColor);
+
+	/** Set per line material Line Size/Thickness parameter that is uniform for all lines */
+	void SetLineThickness(float InSize);
+
 protected:
 
 	//~ UMeshComponent Interface.
diff --git a/Engine/Plugins/Editor/UVEditor/Source/UVEditorTools/Public/Drawing/BasicPointSetComponent.h b/Engine/Plugins/Editor/UVEditor/Source/UVEditorTools/Public/Drawing/BasicPointSetComponent.h
--- a/Engine/Plugins/Editor/UVEditor/Source/UVEditorTools/Public/Drawing/BasicPointSetComponent.h
+++ b/Engine/Plugins/Editor/UVEditor/Source/UVEditorTools/Public/Drawing/BasicPointSetComponent.h
@@ -44,6 +44,12 @@ public:
 	/** Set per point material parameters that are uniform for all points. */
 	void SetPointSetParameters(FColor InColor, float InSize, float InDepthBias);
 
+	/** Set per point material Color parameter that is uniform for all points. */
+	void SetColor(const FColor InColor);
+
+	/** Set per point material Size parameter tat is uniform for all points. */
+	void SetPointSize(const float InSize);
+
 protected:
 
 	//~ UMeshComponent Interface.
diff --git a/Engine/Plugins/Editor/UVEditor/Source/UVEditorTools/Public/Selection/UVToolSelectionAPI.h b/Engine/Plugins/Editor/UVEditor/Source/UVEditorTools/Public/Selection/UVToolSelectionAPI.h
--- a/Engine/Plugins/Editor/UVEditor/Source/UVEditorTools/Public/Selection/UVToolSelectionAPI.h
+++ b/Engine/Plugins/Editor/UVEditor/Source/UVEditorTools/Public/Selection/UVToolSelectionAPI.h
@@ -340,6 +340,18 @@ public:
 	// UUVToolContextObject
 	virtual void Shutdown() override;
 	virtual void OnToolEnded(UInteractiveTool* DeadTool) override;
+
+	// Holds the properties of the UV Live Preview Selection (in 3d viewport) for customization
+	struct FLivePreviewSelectionUXSettings
+	{
+		TOptional<FColor> SelectionColor;
+		TOptional<float> LineThickness;
+		TOptional<float> PointSize;
+	};
+
+	// Sets the UV Live Preview Selection settings
+	void SetLivePreviewSelectionUXSettings(const FLivePreviewSelectionUXSettings& Settings) const;
+
 protected:
 	TArray<FUVToolSelection> CurrentSelections;
 	TArray<FUVToolSelection> CurrentUnsetSelections;
diff --git a/Engine/Plugins/Editor/UVEditor/Source/UVEditorTools/Public/Selection/UVToolSelectionHighlightMechanic.h b/Engine/Plugins/Editor/UVEditor/Source/UVEditorTools/Public/Selection/UVToolSelectionHighlightMechanic.h
--- a/Engine/Plugins/Editor/UVEditor/Source/UVEditorTools/Public/Selection/UVToolSelectionHighlightMechanic.h
+++ b/Engine/Plugins/Editor/UVEditor/Source/UVEditorTools/Public/Selection/UVToolSelectionHighlightMechanic.h
@@ -91,6 +91,21 @@ public:
 	 */
 	void SetEnablePairedEdgeHighlights(bool bEnable);
 
+	/**
+	 * Directly set the Color property of the selection line and point set.
+	 */
+	void SetColor(const FColor Color) const;
+
+	/**
+	 * Directly set the Size/Thickness property of the selection line set.
+	 */
+	void SetLineThickness(const float LineSize) const;
+
+	/**
+	 * directly set the Size property of the selection point set.
+	 */
+	void SetPointSize(const float PointSize) const;
+
 protected:
 
 	UPROPERTY()
