UVEditor: Fix bug with resetting mouse position in live preview when a drag is captured by a tool.

diff --git a/Engine/Plugins/Editor/UVEditor/Source/UVEditor/Private/UVEditor3DViewportClient.cpp b/Engine/Plugins/Editor/UVEditor/Source/UVEditor/Private/UVEditor3DViewportClient.cpp
--- a/Engine/Plugins/Editor/UVEditor/Source/UVEditor/Private/UVEditor3DViewportClient.cpp
+++ b/Engine/Plugins/Editor/UVEditor/Source/UVEditor/Private/UVEditor3DViewportClient.cpp
@@ -22,6 +22,37 @@ FUVEditor3DViewportClient::FUVEditor3DViewportClient(FEditorModeTools* InModeToo
 	OverrideNearClipPlane(KINDA_SMALL_NUMBER);
 }
 
+bool FUVEditor3DViewportClient::ShouldOrbitCamera() const
+{
+	// bIsTracking indicates that the viewport has captured the mouse for camera movement. If we don't have this
+	//  check, things do mostly work (the camera won't orbit unnecessarily), but we will run into a tricky bug
+	//  with tool drag captures: in EUVEditor3DViewportClientCameraMode::Orbit mode, the drag captures will hide
+	//  the cursor because ShouldOrbitCamera would return true, and then the cursor will then be reinstated at the
+	//  old position on capture end, leading to a broken-looking reset of the cursor position.
+	if (!bIsTracking)
+	{
+		return false;
+	}
+
+	// Including some additional checks to prevent the orbit mode from being on all the time,
+	// which ultimately causes weirdness in how the camera transform matrices are handled by
+	// the viewport internally.
+	const bool bLeftMouseButtonDown = Viewport->KeyState(EKeys::LeftMouseButton) && !bLockFlightCamera;
+	const bool bMiddleMouseButtonDown = Viewport->KeyState(EKeys::MiddleMouseButton);
+	const bool bRightMouseButtonDown = Viewport->KeyState(EKeys::RightMouseButton);
+	const bool bIsOnlyAltPressed = IsAltPressed() && !IsCtrlPressed() && !IsShiftPressed();
+
+	switch (CameraMode) {
+	case EUVEditor3DViewportClientCameraMode::Orbit:
+		return bIsOnlyAltPressed || bLeftMouseButtonDown || bMiddleMouseButtonDown || bRightMouseButtonDown;
+	case EUVEditor3DViewportClientCameraMode::Fly:
+		return FEditorViewportClient::ShouldOrbitCamera();
+	default:
+		ensure(false);
+		return FEditorViewportClient::ShouldOrbitCamera();
+	}
+}
+
 void FUVEditor3DViewportClient::FocusCameraOnSelection()
 {
 	if (ViewportButtonsAPI) 
diff --git a/Engine/Plugins/Editor/UVEditor/Source/UVEditor/Public/UVEditor3DViewportClient.h b/Engine/Plugins/Editor/UVEditor/Source/UVEditor/Public/UVEditor3DViewportClient.h
--- a/Engine/Plugins/Editor/UVEditor/Source/UVEditor/Public/UVEditor3DViewportClient.h
+++ b/Engine/Plugins/Editor/UVEditor/Source/UVEditor/Public/UVEditor3DViewportClient.h
@@ -28,26 +28,7 @@ public:
 	virtual ~FUVEditor3DViewportClient() {}
 
 	// FEditorViewportClient
-	virtual bool ShouldOrbitCamera() const override {
-
-		// Including some additional checks to prevent the orbit mode from being on all the time,
-		// which ultimately causes weirdness in how the camera transform matrices are handled by
-		// the viewport internally.
-		const bool bLeftMouseButtonDown = Viewport->KeyState(EKeys::LeftMouseButton) && !bLockFlightCamera;
-		const bool bMiddleMouseButtonDown = Viewport->KeyState(EKeys::MiddleMouseButton);
-		const bool bRightMouseButtonDown = Viewport->KeyState(EKeys::RightMouseButton);
-		const bool bIsOnlyAltPressed = IsAltPressed() && !IsCtrlPressed() && !IsShiftPressed();
-
-		switch (CameraMode) {
-		case EUVEditor3DViewportClientCameraMode::Orbit:
-			return bIsOnlyAltPressed || bLeftMouseButtonDown || bMiddleMouseButtonDown || bRightMouseButtonDown;
-		case EUVEditor3DViewportClientCameraMode::Fly:
-			return FEditorViewportClient::ShouldOrbitCamera();
-		default:
-			ensure(false);
-			return FEditorViewportClient::ShouldOrbitCamera();
-		}
-	}
+	virtual bool ShouldOrbitCamera() const override;
 	bool CanSetWidgetMode(UE::Widget::EWidgetMode NewMode) const override {	return false; }
 	void SetWidgetMode(UE::Widget::EWidgetMode NewMode) override {}
 	UE::Widget::EWidgetMode GetWidgetMode() const override { return UE::Widget::EWidgetMode::WM_None; }
