UV Editor: Support tool hotkeys in the live preview. Support gizmos in the live preview. Properly reinitialize select tool after one-off actions. Make it possible to use one-off actions from additional tools (for upcoming brush tool).

diff --git a/Engine/Plugins/Editor/UVEditor/Source/UVEditor/Private/UVEditor3DViewportMode.cpp b/Engine/Plugins/Editor/UVEditor/Source/UVEditor/Private/UVEditor3DViewportMode.cpp
new file mode 100644
--- /dev/null
+++ b/Engine/Plugins/Editor/UVEditor/Source/UVEditor/Private/UVEditor3DViewportMode.cpp
@@ -0,0 +1,18 @@
+// Copyright Epic Games, Inc. All Rights Reserved.
+
+#include "UVEditor3DViewportMode.h"
+
+#define LOCTEXT_NAMESPACE "UUVEditor3DViewportMode"
+
+const FEditorModeID UUVEditor3DViewportMode::EM_ModeID = TEXT("EM_UVEditor3DViewportModeId");
+
+UUVEditor3DViewportMode::UUVEditor3DViewportMode()
+{
+	Info = FEditorModeInfo(
+		EM_ModeID,
+		LOCTEXT("ModeName", "UV 3D Viewport Dummy Mode"),
+		FSlateIcon(),
+		false);
+}
+
+#undef LOCTEXT_NAMESPACE
\ No newline at end of file
diff --git a/Engine/Plugins/Editor/UVEditor/Source/UVEditor/Private/UVEditor3DViewportMode.h b/Engine/Plugins/Editor/UVEditor/Source/UVEditor/Private/UVEditor3DViewportMode.h
new file mode 100644
--- /dev/null
+++ b/Engine/Plugins/Editor/UVEditor/Source/UVEditor/Private/UVEditor3DViewportMode.h
@@ -0,0 +1,29 @@
+// Copyright Epic Games, Inc. All Rights Reserved.
+
+#pragma once
+
+#include "Tools/UEdMode.h"
+
+#include "UVEditor3DViewportMode.generated.h"
+
+/**
+ * The UV editor live preview has its own world, mode manager, and input router. It doesn't
+ *  really need a mode, and we can't hook it up to the UV editor mode because we don't want
+ *  input to be routed directly to it from this other world (it is instead accessed through
+ *  a context API in tools). However, we do want hotkey events to be routed to our tools from
+ *  the live preview. This is done by registering our command objects to a command list that
+ *  the live preview routes input to, and to have that command list be at the expected place
+ *  in the input routing, we need a mode that has a toolkit object (See FEditorModeTools::InputKey.
+ *  The hotkey handling passes through the ProcessCommandBindings on the mode toolkit there).
+ *  The default mode (UEdModeDefault) does not have a toolkit object, hence this dummy mode.
+ */
+UCLASS(Transient)
+class UUVEditor3DViewportMode : public UEdMode
+{
+	GENERATED_BODY()
+
+public:
+	const static FEditorModeID EM_ModeID;
+
+	UUVEditor3DViewportMode();
+};
diff --git a/Engine/Plugins/Editor/UVEditor/Source/UVEditor/Private/UVEditorCommands.cpp b/Engine/Plugins/Editor/UVEditor/Source/UVEditor/Private/UVEditorCommands.cpp
--- a/Engine/Plugins/Editor/UVEditor/Source/UVEditor/Private/UVEditorCommands.cpp
+++ b/Engine/Plugins/Editor/UVEditor/Source/UVEditor/Private/UVEditorCommands.cpp
@@ -33,13 +33,17 @@ void FUVEditorCommands::RegisterCommands()
 	UI_COMMAND(BeginChannelEditTool, "Channels", "Modify UV channels", EUserInterfaceActionType::ToggleButton, FInputChord());
 	UI_COMMAND(BeginSeamTool, "Seam", "Edit UV seams", EUserInterfaceActionType::ToggleButton, FInputChord());
 	UI_COMMAND(BeginRecomputeUVsTool, "Unwrap", "Perform UV unwrapping", EUserInterfaceActionType::ToggleButton, FInputChord());
+	UI_COMMAND(BeginBrushSelectTool, "Brush", "Brush select triangles", EUserInterfaceActionType::ToggleButton, FInputChord());
 
-	// These currently get linked to actions inside the select tool, but will eventually have their own buttons among the tools
-	// once selection is pulled out to mode-level.
+	// These get linked to one-off tool actions.
 	UI_COMMAND(SewAction, "Sew", "Sew edges highlighted in red to edges highlighted in green", EUserInterfaceActionType::Button, FInputChord());
 	UI_COMMAND(SplitAction, "Split",
 	           "Given an edge selection, split those edges. Given a vertex selection, split any selected bowtie vertices. Given a triangle selection, split along selection boundaries.",
 	           EUserInterfaceActionType::Button, FInputChord());
+	UI_COMMAND(MakeIslandAction, "Island", "Given a triangle selection, make the selection into a single separate UV Island.",
+		EUserInterfaceActionType::Button, FInputChord());
+	UI_COMMAND(UnsetUVsAction, "UnsetUVs", "Unset the UVs on the given triangle selection.",
+		EUserInterfaceActionType::Button, FInputChord());
 
 	// These allow us to link up to pressed keys
 	UI_COMMAND(AcceptOrCompleteActiveTool, "Accept", "Accept the active tool", EUserInterfaceActionType::Button, FInputChord(EKeys::Enter));
@@ -61,5 +65,41 @@ void FUVEditorCommands::RegisterCommands()
 	           FInputChord(EModifierKey::Alt, EKeys::B));
 }
 
+//~ Modeled on ModelingToolsActions.cpp
+UE::Geometry::FUVEditorToolActionCommands::FUVEditorToolActionCommands() :
+	TInteractiveToolCommands<FUVEditorToolActionCommands>(
+		"UVEditorHotkeys", // Context name for fast lookup
+		LOCTEXT("HotkeysCategory", "UV Editor Hotkeys"), // Localized context name for displaying
+		NAME_None, // Parent
+		FUVEditorStyle::Get().GetStyleSetName() // Icon Style Set
+	)
+{
+}
+void UE::Geometry::FUVEditorToolActionCommands::GetToolDefaultObjectList(TArray<UInteractiveTool*>& ToolCDOs)
+{
+}
+void UE::Geometry::FUVEditorToolActionCommands::RegisterAllToolActions()
+{
+}
+void UE::Geometry::FUVEditorToolActionCommands::UnregisterAllToolActions()
+{
+}
+void UE::Geometry::FUVEditorToolActionCommands::UpdateToolCommandBinding(UInteractiveTool* Tool, TSharedPtr<FUICommandList> UICommandList, bool bUnbind)
+{
+#define UPDATE_BINDING(CommandsType)  if (!bUnbind) \
+	CommandsType::Get().BindCommandsForCurrentTool(UICommandList, Tool); \
+	else CommandsType::Get().UnbindActiveCommands(UICommandList);
+
+}
+
+//~ Modeled on ModelingToolsActions.cpp
+#define DEFINE_TOOL_ACTION_COMMANDS(CommandsClassName, ContextNameString, SettingsDialogString, ToolClassName ) \
+UE::Geometry::CommandsClassName::CommandsClassName() : TInteractiveToolCommands<CommandsClassName>( \
+ContextNameString, NSLOCTEXT("Contexts", ContextNameString, SettingsDialogString), NAME_None, FUVEditorStyle::Get().GetStyleSetName()) {} \
+void UE::Geometry::CommandsClassName::GetToolDefaultObjectList(TArray<UInteractiveTool*>& ToolCDOs) \
+{\
+	ToolCDOs.Add(GetMutableDefault<ToolClassName>()); \
+}
+
 
-#undef LOCTEXT_NAMESPACE
+#undef LOCTEXT_NAMESPACE
\ No newline at end of file
diff --git a/Engine/Plugins/Editor/UVEditor/Source/UVEditor/Private/UVEditorMode.cpp b/Engine/Plugins/Editor/UVEditor/Source/UVEditor/Private/UVEditorMode.cpp
--- a/Engine/Plugins/Editor/UVEditor/Source/UVEditor/Private/UVEditorMode.cpp
+++ b/Engine/Plugins/Editor/UVEditor/Source/UVEditor/Private/UVEditorMode.cpp
@@ -14,6 +14,7 @@
 #include "Selection/UVToolSelectionAPI.h"
 #include "Drawing/MeshElementsVisualizer.h"
 #include "Editor.h"
+#include "EditorModes.h"
 #include "EditorViewportClient.h"
 #include "EdModeInteractiveToolsContext.h" //ToolsContext, EditorInteractiveToolsContext
 #include "EngineAnalytics.h"
@@ -31,6 +32,7 @@
 #include "ToolTargets/UVEditorToolMeshInput.h"
 #include "ToolTargetManager.h"
 #include "ToolTargets/UVEditorToolMeshInput.h"
+#include "UVEditor3DViewportMode.h"
 #include "UVEditorCommands.h"
 #include "UVEditorLayoutTool.h"
 #include "UVEditorTransformTool.h"
@@ -382,6 +384,8 @@ void UUVEditorMode::Enter()
 				SelectionAPI->LivePreviewDrawHUD(Canvas, RenderAPI);
 			}
 		});
+
+		LivePreviewToolkitCommands = InitContext->LivePreviewToolkitCommands;
 	}
 
 	InitializeModeContexts();
@@ -577,6 +581,7 @@ void UUVEditorMode::RegisterTools()
 	UVSelectToolBuilder->Targets = &ToolInputObjects;
 	DefaultToolIdentifier = TEXT("BeginSelectTool");
 	GetToolManager()->RegisterToolType(DefaultToolIdentifier, UVSelectToolBuilder);
+	ToolsThatAllowActions.Add(DefaultToolIdentifier);
 
 	// Note that the identifiers below need to match the command names so that the tool icons can 
 	// be easily retrieved from the active tool name in UVEditorModeToolkit::OnToolStarted. Otherwise
@@ -621,6 +626,8 @@ void UUVEditorMode::RegisterTools()
 
 void UUVEditorMode::RegisterActions()
 {
+	using namespace UVEditorModeLocals;
+
 	const FUVEditorCommands& CommandInfos = FUVEditorCommands::Get();
 	const TSharedRef<FUICommandList>& CommandList = Toolkit->GetToolkitCommands();
 
@@ -629,7 +636,10 @@ void UUVEditorMode::RegisterActions()
 		Action->Setup(GetToolManager());
 		CommandList->MapAction(CommandInfo,
 			FExecuteAction::CreateWeakLambda(Action, [this, Action]() 
-			{ 
+			{
+				// If we're activating from some other selection tool, go ahead and switch to the regular one
+				ActivateDefaultTool();
+
 				Action->ExecuteAction();
 				for (const FUVToolSelection& Selection : SelectionAPI->GetSelections())
 				{
@@ -647,9 +657,19 @@ void UUVEditorMode::RegisterActions()
 				}
 			}),
 			FCanExecuteAction::CreateWeakLambda(Action, [Action, this]() 
-				{ 
-					return IsDefaultToolActive() && Action->CanExecuteAction(); 
-				}));
+			{
+				// We could make some generic system for determining which tools are safe to call one-off actions from,
+				//  but it's unclear that it's worth it while we have so few.
+				auto DoesCurrentToolAllowActions = [this]()
+				{
+					if (UInteractiveToolsContext* ToolsContext = GetInteractiveToolsContext())
+					{
+						return ToolsThatAllowActions.Contains(ToolsContext->GetActiveToolName(EToolSide::Mouse));
+					}
+					return false;
+				};
+				return DoesCurrentToolAllowActions() && Action->CanExecuteAction();
+			}));
 		RegisteredActions.Add(Action);
 	};
 	PrepAction(CommandInfos.SewAction, NewObject<UUVSeamSewAction>());
@@ -681,6 +701,12 @@ void UUVEditorMode::OnToolStarted(UInteractiveToolManager* Manager, UInteractive
 {
 	using namespace UVEditorModeLocals;
 
+	FUVEditorToolActionCommands::UpdateToolCommandBinding(Tool, Toolkit->GetToolkitCommands(), false);
+	if (LivePreviewToolkitCommands.IsValid())
+	{
+		FUVEditorToolActionCommands::UpdateToolCommandBinding(Tool, LivePreviewToolkitCommands.Pin(), false);
+	}
+
 	FText TransactionName = LOCTEXT("ActivateTool", "Activate Tool");
 	
 	if (!IsDefaultToolActive())
@@ -727,6 +753,12 @@ void UUVEditorMode::OnToolStarted(UInteractiveToolManager* Manager, UInteractive
 
 void UUVEditorMode::OnToolEnded(UInteractiveToolManager* Manager, UInteractiveTool* Tool)
 {
+	FUVEditorToolActionCommands::UpdateToolCommandBinding(Tool, Toolkit->GetToolkitCommands(), true);
+	if (LivePreviewToolkitCommands.IsValid())
+	{
+		FUVEditorToolActionCommands::UpdateToolCommandBinding(Tool, LivePreviewToolkitCommands.Pin(), true);
+	}
+
 	for (TWeakObjectPtr<UUVToolContextObject> Context : ContextsToUpdateOnToolEnd)
 	{
 		if (ensure(Context.IsValid()))
@@ -1061,8 +1093,8 @@ void UUVEditorMode::InitializeAssetEditorContexts(UContextObjectStore& ContextSt
 				{
 					LivePreviewToolsContextCursorAPI->ClearCursorOverride();
 				}
-			}
-			);
+			},
+			LivePreviewModeManager.GetInteractiveToolsContext()->GizmoManager);
 		ContextStore.AddContextObject(LivePreviewAPI);
 	}
 
@@ -1071,6 +1103,18 @@ void UUVEditorMode::InitializeAssetEditorContexts(UContextObjectStore& ContextSt
 	{
 		UUVEditorInitializationContext* InitContext = NewObject<UUVEditorInitializationContext>();
 		InitContext->LivePreviewITC = Cast<UEditorInteractiveToolsContext>(LivePreviewModeManager.GetInteractiveToolsContext());
+
+		LivePreviewModeManager.ActivateMode(UUVEditor3DViewportMode::EM_ModeID);
+		UEdMode* LivePreviewDefaultMode = LivePreviewModeManager.GetActiveScriptableMode(UUVEditor3DViewportMode::EM_ModeID);
+		if (ensure(LivePreviewDefaultMode))
+		{
+			TSharedPtr<FModeToolkit> Toolkit = LivePreviewDefaultMode->GetToolkit().Pin();
+			if (ensure(Toolkit.IsValid()))
+			{
+				InitContext->LivePreviewToolkitCommands = Toolkit->GetToolkitCommands();
+			}
+		}
+
 		ContextStore.AddContextObject(InitContext);
 	}
 
diff --git a/Engine/Plugins/Editor/UVEditor/Source/UVEditor/Private/UVEditorModeToolkit.cpp b/Engine/Plugins/Editor/UVEditor/Source/UVEditor/Private/UVEditorModeToolkit.cpp
--- a/Engine/Plugins/Editor/UVEditor/Source/UVEditor/Private/UVEditorModeToolkit.cpp
+++ b/Engine/Plugins/Editor/UVEditor/Source/UVEditor/Private/UVEditorModeToolkit.cpp
@@ -31,6 +31,10 @@
 
 namespace UVEditorModeToolkitLocals
 {
+	TAutoConsoleVariable<bool> CVarAddDebugTools(
+		TEXT("modeling.UVEditor.AddDebugTools"),
+		false,
+		TEXT("Enable UV tools that can be useful for debugging (currently just the \"Unset UVs\" tool."));
 }
 
 FUVEditorModeToolkit::FUVEditorModeToolkit()
@@ -465,6 +469,8 @@ FText FUVEditorModeToolkit::GetToolPaletteDisplayName(FName Palette) const
 
 void FUVEditorModeToolkit::BuildToolPalette(FName PaletteIndex, class FToolBarBuilder& ToolbarBuilder)
 {
+	using namespace UVEditorModeToolkitLocals;
+
 	const FUVEditorCommands& Commands = FUVEditorCommands::Get();
 
 	if (PaletteIndex == ToolsTabName)
diff --git a/Engine/Plugins/Editor/UVEditor/Source/UVEditor/Private/UVEditorModule.cpp b/Engine/Plugins/Editor/UVEditor/Source/UVEditor/Private/UVEditorModule.cpp
--- a/Engine/Plugins/Editor/UVEditor/Source/UVEditor/Private/UVEditorModule.cpp
+++ b/Engine/Plugins/Editor/UVEditor/Source/UVEditor/Private/UVEditorModule.cpp
@@ -9,6 +9,7 @@
 #include "PropertyEditorModule.h"
 #include "Selection.h"
 #include "UVEditor.h"
+#include "UVEditor3DViewportMode.h"
 #include "UVEditorCommands.h"
 #include "UVEditorMode.h"
 #include "UVEditorStyle.h"
@@ -30,6 +31,7 @@ void FUVEditorModule::StartupModule()
 
 	FUVEditorStyle::Get(); // Causes the constructor to be called
 	FUVEditorCommands::Register();
+	UE::Geometry::FUVEditorToolActionCommands::RegisterAllToolActions();
 
 	// Menus need to be registered in a callback to make sure the system is ready for them.
 	UToolMenus::RegisterStartupCallback(
@@ -56,9 +58,11 @@ void FUVEditorModule::ShutdownModule()
 	UToolMenus::UnRegisterStartupCallback(this);
 	UToolMenus::UnregisterOwner(this);
 
+	UE::Geometry::FUVEditorToolActionCommands::UnregisterAllToolActions();
 	FUVEditorCommands::Unregister();
 
 	FEditorModeRegistry::Get().UnregisterMode(UUVEditorMode::EM_UVEditorModeId);
+	FEditorModeRegistry::Get().UnregisterMode(UUVEditor3DViewportMode::EM_ModeID);
 
 	// Unregister customizations
 	FPropertyEditorModule* PropertyEditorModule = FModuleManager::GetModulePtr<FPropertyEditorModule>("PropertyEditor");
diff --git a/Engine/Plugins/Editor/UVEditor/Source/UVEditor/Private/UVEditorToolkit.cpp b/Engine/Plugins/Editor/UVEditor/Source/UVEditor/Private/UVEditorToolkit.cpp
--- a/Engine/Plugins/Editor/UVEditor/Source/UVEditor/Private/UVEditorToolkit.cpp
+++ b/Engine/Plugins/Editor/UVEditor/Source/UVEditor/Private/UVEditorToolkit.cpp
@@ -24,6 +24,7 @@
 #include "UVEditorSubsystem.h"
 #include "UVEditorModule.h"
 #include "UVEditorStyle.h"
+#include "UVEditor3DViewportMode.h"
 #include "ContextObjects/UVToolContextObjects.h"
 #include "UVEditorModeUILayer.h"
 #include "Widgets/Docking/SDockTab.h"
@@ -109,6 +110,7 @@ FUVEditorToolkit::FUVEditorToolkit(UAssetEditor* InOwningAssetEditor)
 
 	LivePreviewEditorModeManager = MakeShared<FAssetEditorModeManager>();
 	LivePreviewEditorModeManager->SetPreviewScene(LivePreviewScene.Get());
+	LivePreviewEditorModeManager->SetDefaultMode(UUVEditor3DViewportMode::EM_ModeID);
 	LivePreviewInputRouter = LivePreviewEditorModeManager->GetInteractiveToolsContext()->InputRouter;
 
 	LivePreviewTabContent = MakeShareable(new FEditorViewportTabContent());
@@ -297,7 +299,7 @@ void FUVEditorToolkit::OnClose()
 	// This is super important to do, otherwise currently opened tabs won't be marked as "closed".
 	// This results in tabs not being properly recycled upon reopening the editor and tab
 	// duplication for each opening event.
-	GetEditorModeManager().ActivateDefaultMode();
+	GetEditorModeManager().DeactivateAllModes();
 
 	FAssetEditorToolkit::OnClose();
 }
@@ -438,10 +440,12 @@ void FUVEditorToolkit::PostInitAssetEditor()
 	ModeUILayer = MakeShareable(new FUVEditorModeUILayer(PinnedToolkitHost.Get()));
 	ModeUILayer->SetModeMenuCategory( UVEditorMenuCategory );
 
+	// Needed so that the live preview dummy mode is initialized enough to be able to route hotkeys
+	LivePreviewEditorModeManager->SetToolkitHost(PinnedToolkitHost.ToSharedRef());
+
 	TArray<TObjectPtr<UObject>> ObjectsToEdit;
 	OwningAssetEditor->GetObjectsToEdit(MutableView(ObjectsToEdit));
 
-	// TODO: get these when possible (from level editor selection, for instance), and set them to something reasonable otherwise.
 	TArray<FTransform> ObjectTransforms;
 	Cast<UUVEditor>(OwningAssetEditor)->GetWorldspaceRelativeTransforms(ObjectTransforms);
 
diff --git a/Engine/Plugins/Editor/UVEditor/Source/UVEditor/Public/UVEditorCommands.h b/Engine/Plugins/Editor/UVEditor/Source/UVEditor/Public/UVEditorCommands.h
--- a/Engine/Plugins/Editor/UVEditor/Source/UVEditor/Public/UVEditorCommands.h
+++ b/Engine/Plugins/Editor/UVEditor/Source/UVEditor/Public/UVEditorCommands.h
@@ -4,6 +4,7 @@
 
 #include "CoreMinimal.h"
 #include "Framework/Commands/Commands.h"
+#include "Tools/InteractiveToolsCommands.h"
 
 class UVEDITOR_API FUVEditorCommands : public TCommands<FUVEditorCommands>
 {
@@ -27,9 +28,12 @@ public:
 	TSharedPtr<FUICommandInfo> BeginChannelEditTool;
 	TSharedPtr<FUICommandInfo> BeginSeamTool;
 	TSharedPtr<FUICommandInfo> BeginRecomputeUVsTool;
+	TSharedPtr<FUICommandInfo> BeginBrushSelectTool;
 
 	TSharedPtr<FUICommandInfo> SewAction;
 	TSharedPtr<FUICommandInfo> SplitAction;
+	TSharedPtr<FUICommandInfo> MakeIslandAction;
+	TSharedPtr<FUICommandInfo> UnsetUVsAction;
 
 	TSharedPtr<FUICommandInfo> AcceptOrCompleteActiveTool;
 	TSharedPtr<FUICommandInfo> CancelOrCompleteActiveTool;
@@ -46,4 +50,38 @@ public:
 	TSharedPtr<FUICommandInfo> SetFocusCamera;
 
 	TSharedPtr<FUICommandInfo> ToggleBackground;
-};
\ No newline at end of file
+};
+
+namespace UE::Geometry
+{
+//~ Modeled on ModelingToolsActions.h
+class FUVEditorToolActionCommands : public TInteractiveToolCommands<FUVEditorToolActionCommands>
+{
+public:
+	FUVEditorToolActionCommands();
+
+	virtual void GetToolDefaultObjectList(TArray<UInteractiveTool*>& ToolCDOs) override;
+	static void RegisterAllToolActions();
+	static void UnregisterAllToolActions();
+
+	/**
+	 * Add or remove commands relevant to Tool to the given UICommandList.
+	 * Call this when the active tool changes (eg on ToolManager.OnToolStarted / OnToolEnded)
+	 * @param bUnbind if true, commands are removed, otherwise added
+	 */
+	static void UpdateToolCommandBinding(UInteractiveTool* Tool, TSharedPtr<FUICommandList> UICommandList, bool bUnbind = false);
+};
+}
+
+//~ Modeled on ModelingToolsActions.h
+#define DECLARE_TOOL_ACTION_COMMANDS(CommandsClassName) \
+class CommandsClassName : public TInteractiveToolCommands<CommandsClassName> \
+{\
+public:\
+	CommandsClassName();\
+	virtual void GetToolDefaultObjectList(TArray<UInteractiveTool*>& ToolCDOs) override;\
+};\
+
+namespace UE::Geometry
+{
+}
\ No newline at end of file
diff --git a/Engine/Plugins/Editor/UVEditor/Source/UVEditor/Public/UVEditorInitializationContext.h b/Engine/Plugins/Editor/UVEditor/Source/UVEditor/Public/UVEditorInitializationContext.h
--- a/Engine/Plugins/Editor/UVEditor/Source/UVEditor/Public/UVEditorInitializationContext.h
+++ b/Engine/Plugins/Editor/UVEditor/Source/UVEditor/Public/UVEditorInitializationContext.h
@@ -8,6 +8,7 @@
 
 #include "UVEditorInitializationContext.generated.h"
 
+class FUICommandList;
 class UEditorInteractiveToolsContext;
 
 /**
@@ -24,4 +25,5 @@ class UVEDITOR_API UUVEditorInitializationContext : public UUVToolContextObject
 	GENERATED_BODY()
 public:
 	TWeakObjectPtr<UEditorInteractiveToolsContext> LivePreviewITC;
+	TWeakPtr<FUICommandList> LivePreviewToolkitCommands;
 };
\ No newline at end of file
diff --git a/Engine/Plugins/Editor/UVEditor/Source/UVEditor/Public/UVEditorMode.h b/Engine/Plugins/Editor/UVEditor/Source/UVEditor/Public/UVEditorMode.h
--- a/Engine/Plugins/Editor/UVEditor/Source/UVEditor/Public/UVEditorMode.h
+++ b/Engine/Plugins/Editor/UVEditor/Source/UVEditor/Public/UVEditorMode.h
@@ -14,6 +14,7 @@
 PREDECLARE_GEOMETRY(class FDynamicMesh3);
 
 class FAssetEditorModeManager;
+class FUICommandList;
 class UEditorInteractiveToolsContext;
 class FEditorViewportClient;
 class FToolCommandChange;
@@ -369,6 +370,13 @@ protected:
 	// Used to forward Render/DrawHUD calls in the live preview to the api object.
 	TWeakObjectPtr<UEditorInteractiveToolsContext> LivePreviewITC;
 
+	//~ Other members should be private as well, but too late for that. Keeping this one here
+	//~  to be next to the other live preview members.
+private: 
+	// Used to make mode hotkeys be triggered properly from the live preview
+	TWeakPtr<FUICommandList> LivePreviewToolkitCommands;
+protected:
+
 	UPROPERTY()
 	TObjectPtr<UUVToolSelectionAPI> SelectionAPI = nullptr;
 
@@ -404,5 +412,6 @@ private:
 	UPROPERTY()
 	TObjectPtr<UUVEditorLivePreviewUXProperties> UVEditorLivePreviewUXProperties = nullptr;
 
+	TArray<FString> ToolsThatAllowActions;
 };
 
diff --git a/Engine/Plugins/Editor/UVEditor/Source/UVEditorTools/Private/ContextObjects/UVToolContextObjects.cpp b/Engine/Plugins/Editor/UVEditor/Source/UVEditorTools/Private/ContextObjects/UVToolContextObjects.cpp
--- a/Engine/Plugins/Editor/UVEditor/Source/UVEditorTools/Private/ContextObjects/UVToolContextObjects.cpp
+++ b/Engine/Plugins/Editor/UVEditor/Source/UVEditorTools/Private/ContextObjects/UVToolContextObjects.cpp
@@ -5,6 +5,7 @@
 #include "DynamicMesh/MeshIndexUtil.h"
 #include "Engine/World.h"
 #include "InputRouter.h" // Need to define this and UWorld so weak pointers know they are UThings
+#include "InteractiveGizmoManager.h"
 #include "InteractiveToolManager.h"
 #include "ToolTargets/UVEditorToolMeshInput.h"
 #include "UDIMUtilities.h"
@@ -15,7 +16,6 @@ using namespace UE::Geometry;
 
 namespace UVToolContextObjectLocals
 {
-
 	/**
 	 * A wrapper change that applies a given change to the unwrap canonical mesh of an input, and uses that
 	 * to update the other views. Causes a broadcast of OnCanonicalModified.
@@ -81,13 +81,15 @@ void UUVToolEmitChangeAPI::EmitToolDependentChange(UObject* TargetObject, TUniqu
 void UUVToolLivePreviewAPI::Initialize(UWorld* WorldIn, UInputRouter* RouterIn,
 	TUniqueFunction<void(FViewCameraState& CameraStateOut)> GetLivePreviewCameraStateFuncIn,
 	TUniqueFunction<void(const FAxisAlignedBox3d& BoundingBox)> SetLivePreviewCameraToLookAtVolumeFuncIn,
-	TUniqueFunction<void(const EMouseCursor::Type Cursor, bool bEnableOverride)> SetCursorOverrideFuncIn)
+	TUniqueFunction<void(const EMouseCursor::Type Cursor, bool bEnableOverride)> SetCursorOverrideFuncIn, 
+	UInteractiveGizmoManager* GizmoManagerIn)
 {
 	World = WorldIn;
 	InputRouter = RouterIn;
 	GetLivePreviewCameraStateFunc = MoveTemp(GetLivePreviewCameraStateFuncIn);
 	SetLivePreviewCameraToLookAtVolumeFunc = MoveTemp(SetLivePreviewCameraToLookAtVolumeFuncIn);
 	SetCursorOverrideFunc = MoveTemp(SetCursorOverrideFuncIn);
+	GizmoManager = GizmoManagerIn;
 }
 
 void UUVToolLivePreviewAPI::OnToolEnded(UInteractiveTool* DeadTool)
diff --git a/Engine/Plugins/Editor/UVEditor/Source/UVEditorTools/Private/UVSelectTool.cpp b/Engine/Plugins/Editor/UVEditor/Source/UVEditorTools/Private/UVSelectTool.cpp
--- a/Engine/Plugins/Editor/UVEditor/Source/UVEditorTools/Private/UVSelectTool.cpp
+++ b/Engine/Plugins/Editor/UVEditor/Source/UVEditorTools/Private/UVSelectTool.cpp
@@ -218,6 +218,7 @@ void UUVSelectTool::Setup()
 
 	// Make sure that if undo/redo events act on the meshes, we update our state.
 	// The trees will be updated by the tree store, which listens to the same broadcasts.
+	// This also handles the case of UV editor one-off actions, which don't actually shut down this tool when they run.
 	for (int32 i = 0; i < Targets.Num(); ++i)
 	{
 		Targets[i]->OnCanonicalModified.AddWeakLambda(this, [this]
@@ -227,7 +228,7 @@ void UUVSelectTool::Setup()
 			{
 				// We update the gizmo with 'true' here to have the selection api recompute the
 				// centroid because it may not yet be notified that the mesh has changed.
-				UpdateGizmo(true);
+				ReinitializeFromSelection(true);
 				SelectionAPI->RebuildUnwrapHighlight(TransformGizmo->GetGizmoTransform());
 			}
 		});
@@ -364,7 +365,7 @@ void UUVSelectTool::UpdateGizmo(bool bForceRecomputeSelectionCenters)
 void UUVSelectTool::OnSelectionChanged(bool, uint32 SelectionChangeType)
 {
 	TRACE_CPUPROFILER_EVENT_SCOPE(UVSelectTool_OnSelectionChanged);
-	
+
 	using namespace UVSelectToolLocals;
 
 	if (!(SelectionChangeType & (uint32)ESelectionChangeTypeFlag::SelectionChanged))
@@ -372,6 +373,11 @@ void UUVSelectTool::OnSelectionChanged(bool, uint32 SelectionChangeType)
 		return;
 	}
 
+	ReinitializeFromSelection(false);
+}
+
+void UUVSelectTool::ReinitializeFromSelection(bool bForceRecomputeCachedSelectionCentroid)
+{
 	MovingVidsPerSelection.Reset();
 	MovingVertOriginalPositionsPerSelection.Reset();
 	RenderUpdateTidsPerSelection.Reset();
@@ -473,7 +479,7 @@ void UUVSelectTool::OnSelectionChanged(bool, uint32 SelectionChangeType)
 		RenderUpdateTidsPerSelection.Emplace(TidSet);
 	}
 
-	UpdateGizmo();
+	UpdateGizmo(bForceRecomputeCachedSelectionCentroid);
 }
 
 void UUVSelectTool::GizmoTransformStarted(UTransformProxy* Proxy)
diff --git a/Engine/Plugins/Editor/UVEditor/Source/UVEditorTools/Public/ContextObjects/UVToolContextObjects.h b/Engine/Plugins/Editor/UVEditor/Source/UVEditorTools/Public/ContextObjects/UVToolContextObjects.h
--- a/Engine/Plugins/Editor/UVEditor/Source/UVEditorTools/Public/ContextObjects/UVToolContextObjects.h
+++ b/Engine/Plugins/Editor/UVEditor/Source/UVEditorTools/Public/ContextObjects/UVToolContextObjects.h
@@ -14,6 +14,8 @@
 // TODO: This should be spread out across multiple files
 
 PREDECLARE_GEOMETRY(class FDynamicMesh3);
+class FAssetEditorModeManager;
+class FEditorViewportClient;
 class FToolCommandChange;
 struct FViewCameraState;
 class UInputRouter;
@@ -109,10 +111,12 @@ public:
 	void Initialize(UWorld* WorldIn, UInputRouter* RouterIn,
 		TUniqueFunction<void(FViewCameraState& CameraStateOut)> GetLivePreviewCameraStateFuncIn,
 		TUniqueFunction<void(const UE::Geometry::FAxisAlignedBox3d& BoundingBox)> SetLivePreviewCameraToLookAtVolumeFuncIn,
-		TUniqueFunction<void(const EMouseCursor::Type Cursor, bool bEnableOverride)> SetCursorOverrideFuncIn);
+		TUniqueFunction<void(const EMouseCursor::Type Cursor, bool bEnableOverride)> SetCursorOverrideFuncIn, 
+		UInteractiveGizmoManager* GizmoManager = nullptr);
 
 	UWorld* GetLivePreviewWorld() { return World.Get(); }
 	UInputRouter* GetLivePreviewInputRouter() { return InputRouter.Get(); }
+	UInteractiveGizmoManager* GetGizmoManager() { return GizmoManager.Get(); }
 	void GetLivePreviewCameraState(FViewCameraState& CameraStateOut) 
 	{ 
 		if (GetLivePreviewCameraStateFunc)
@@ -171,6 +175,9 @@ protected:
 	TUniqueFunction<void(FViewCameraState& CameraStateOut)> GetLivePreviewCameraStateFunc;
 	TUniqueFunction<void(const UE::Geometry::FAxisAlignedBox3d& BoundingBox)> SetLivePreviewCameraToLookAtVolumeFunc;
 	TUniqueFunction<void(const EMouseCursor::Type Cursor, bool bEnableOverride)> SetCursorOverrideFunc;
+
+private:
+	TWeakObjectPtr<UInteractiveGizmoManager> GizmoManager;
 };
 
 USTRUCT()
diff --git a/Engine/Plugins/Editor/UVEditor/Source/UVEditorTools/Public/UVSelectTool.h b/Engine/Plugins/Editor/UVEditor/Source/UVEditorTools/Public/UVSelectTool.h
--- a/Engine/Plugins/Editor/UVEditor/Source/UVEditorTools/Public/UVSelectTool.h
+++ b/Engine/Plugins/Editor/UVEditor/Source/UVEditorTools/Public/UVSelectTool.h
@@ -116,5 +116,8 @@ protected:
 	// Inner arrays for these two are 1:1 with each other
 	TArray<TArray<int32>> MovingVidsPerSelection;
 	TArray<TArray<FVector3d>> MovingVertOriginalPositionsPerSelection;
+
+private:
+	void ReinitializeFromSelection(bool bForceRecomputeSelectionCenters);
 };
 
