UVEditor: Fix undo/redo bug dealing with unset triangles.

diff --git a/Engine/Plugins/Editor/UVEditor/Source/UVEditorTools/Private/ToolTargets/UVEditorToolMeshInput.cpp b/Engine/Plugins/Editor/UVEditor/Source/UVEditorTools/Private/ToolTargets/UVEditorToolMeshInput.cpp
--- a/Engine/Plugins/Editor/UVEditor/Source/UVEditorTools/Private/ToolTargets/UVEditorToolMeshInput.cpp
+++ b/Engine/Plugins/Editor/UVEditor/Source/UVEditorTools/Private/ToolTargets/UVEditorToolMeshInput.cpp
@@ -319,11 +319,22 @@ void UUVEditorToolMeshInput::UpdateAllFromAppliedPreview(
 void UUVEditorToolMeshInput::UpdateFromCanonicalUnwrapUsingMeshChange(
 	const FDynamicMeshChange& UnwrapCanonicalMeshChange, bool bBroadcast)
 {
-	// Note that we know that no triangles were created or destroyed since the UV editor
-	// does not allow that (it would break the mesh mappings). Otherwise we would need to
-	// combine original and final tris here.
+	// Get the changed triangle list out of the mesh change. Unfortunately, we need to
+	//  get both the initial and new triangles and lump them together in case unwrap triangles
+	//  got created or destroyed via UV setting and unsetting.
+	// We could maybe get around this because in most cases we know that no triangles got
+	//  created or destroyed, or at least that none got destroyed (because we don't typically
+	//  allow unsetting of triangles), but it would require passing that knowledge down, and
+	//  in the end, it's not worth the trouble for something that only runs on undo/redo.
 	TArray<int32> ChangedTids;
 	UnwrapCanonicalMeshChange.GetSavedTriangleList(ChangedTids, true);
+	TSet<int32> ChangedTidSet(ChangedTids);
+
+	ChangedTids.Reset();
+	UnwrapCanonicalMeshChange.GetSavedTriangleList(ChangedTids, false);
+	ChangedTidSet.Append(ChangedTids);
+
+	ChangedTids = ChangedTidSet.Array();
 
 	TArray<int32> ChangedVids;
 	TriangleToVertexIDs(UnwrapCanonical.Get(), ChangedTids, ChangedVids);
