UVEditor: Add brush selection tool.

diff --git a/Engine/Plugins/Editor/UVEditor/Source/UVEditor/Private/UVEditorCommands.cpp b/Engine/Plugins/Editor/UVEditor/Source/UVEditor/Private/UVEditorCommands.cpp
--- a/Engine/Plugins/Editor/UVEditor/Source/UVEditor/Private/UVEditorCommands.cpp
+++ b/Engine/Plugins/Editor/UVEditor/Source/UVEditor/Private/UVEditorCommands.cpp
@@ -4,6 +4,7 @@
 
 #include "Styling/AppStyle.h"
 #include "Framework/Commands/InputChord.h"
+#include "UVEditorBrushSelectTool.h"
 #include "UVEditorStyle.h"
 
 #define LOCTEXT_NAMESPACE "FUVEditorCommands"
@@ -80,9 +81,11 @@ void UE::Geometry::FUVEditorToolActionCommands::GetToolDefaultObjectList(TArray<
 }
 void UE::Geometry::FUVEditorToolActionCommands::RegisterAllToolActions()
 {
+	UE::Geometry::FUVEditorBrushSelectToolCommands::Register();
 }
 void UE::Geometry::FUVEditorToolActionCommands::UnregisterAllToolActions()
 {
+	UE::Geometry::FUVEditorBrushSelectToolCommands::Unregister();
 }
 void UE::Geometry::FUVEditorToolActionCommands::UpdateToolCommandBinding(UInteractiveTool* Tool, TSharedPtr<FUICommandList> UICommandList, bool bUnbind)
 {
@@ -90,6 +93,10 @@ void UE::Geometry::FUVEditorToolActionCommands::UpdateToolCommandBinding(UIntera
 	CommandsType::Get().BindCommandsForCurrentTool(UICommandList, Tool); \
 	else CommandsType::Get().UnbindActiveCommands(UICommandList);
 
+	if (ExactCast<UUVEditorBrushSelectTool>(Tool))
+	{
+		UPDATE_BINDING(UE::Geometry::FUVEditorBrushSelectToolCommands);
+	}
 }
 
 //~ Modeled on ModelingToolsActions.cpp
@@ -101,5 +108,7 @@ void UE::Geometry::CommandsClassName::GetToolDefaultObjectList(TArray<UInteracti
 	ToolCDOs.Add(GetMutableDefault<ToolClassName>()); \
 }
 
+DEFINE_TOOL_ACTION_COMMANDS(FUVEditorBrushSelectToolCommands, "UVBrushSelect", "UV Editor - Brush Select", UUVEditorBrushSelectTool);
+
 
 #undef LOCTEXT_NAMESPACE
\ No newline at end of file
diff --git a/Engine/Plugins/Editor/UVEditor/Source/UVEditor/Private/UVEditorMode.cpp b/Engine/Plugins/Editor/UVEditor/Source/UVEditor/Private/UVEditorMode.cpp
--- a/Engine/Plugins/Editor/UVEditor/Source/UVEditor/Private/UVEditorMode.cpp
+++ b/Engine/Plugins/Editor/UVEditor/Source/UVEditor/Private/UVEditorMode.cpp
@@ -35,6 +35,8 @@
 #include "ToolTargetManager.h"
 #include "ToolTargets/UVEditorToolMeshInput.h"
 #include "UVEditor3DViewportMode.h"
+#include "UVEditorToolBase.h"
+#include "UVEditorBrushSelectTool.h"
 #include "UVEditorCommands.h"
 #include "UVEditorLayoutTool.h"
 #include "UVEditorTransformTool.h"
@@ -624,6 +626,12 @@ void UUVEditorMode::RegisterTools()
 	UUVEditorRecomputeUVsToolBuilder* UVEditorRecomputeUVsToolBuilder = NewObject<UUVEditorRecomputeUVsToolBuilder>();
 	UVEditorRecomputeUVsToolBuilder->Targets = &ToolInputObjects;
 	RegisterTool(CommandInfos.BeginRecomputeUVsTool, TEXT("BeginRecomputeUVsTool"), UVEditorRecomputeUVsToolBuilder);
+
+	FString BrushToolIdentifier = TEXT("BeginBrushSelectTool");
+	UGenericUVEditorToolBuilder* BrushToolBuilder = NewObject<UGenericUVEditorToolBuilder>();
+	BrushToolBuilder->Initialize(ToolInputObjects, UUVEditorBrushSelectTool::StaticClass());
+	RegisterTool(CommandInfos.BeginBrushSelectTool, BrushToolIdentifier, BrushToolBuilder);
+	ToolsThatAllowActions.Add(BrushToolIdentifier);
 }
 
 void UUVEditorMode::RegisterActions()
diff --git a/Engine/Plugins/Editor/UVEditor/Source/UVEditor/Private/UVEditorModeToolkit.cpp b/Engine/Plugins/Editor/UVEditor/Source/UVEditor/Private/UVEditorModeToolkit.cpp
--- a/Engine/Plugins/Editor/UVEditor/Source/UVEditor/Private/UVEditorModeToolkit.cpp
+++ b/Engine/Plugins/Editor/UVEditor/Source/UVEditor/Private/UVEditorModeToolkit.cpp
@@ -492,6 +492,7 @@ void FUVEditorModeToolkit::BuildToolPalette(FName PaletteIndex, class FToolBarBu
 		ToolbarBuilder.AddToolBarButton(Commands.BeginSeamTool);
 		ToolbarBuilder.AddToolBarButton(Commands.BeginParameterizeMeshTool);
 		ToolbarBuilder.AddToolBarButton(Commands.BeginRecomputeUVsTool);
+		ToolbarBuilder.AddToolBarButton(Commands.BeginBrushSelectTool);
 	}
 }
 
diff --git a/Engine/Plugins/Editor/UVEditor/Source/UVEditor/Public/UVEditorCommands.h b/Engine/Plugins/Editor/UVEditor/Source/UVEditor/Public/UVEditorCommands.h
--- a/Engine/Plugins/Editor/UVEditor/Source/UVEditor/Public/UVEditorCommands.h
+++ b/Engine/Plugins/Editor/UVEditor/Source/UVEditor/Public/UVEditorCommands.h
@@ -84,4 +84,5 @@ public:\
 
 namespace UE::Geometry
 {
+	DECLARE_TOOL_ACTION_COMMANDS(FUVEditorBrushSelectToolCommands);
 }
\ No newline at end of file
diff --git a/Engine/Plugins/Editor/UVEditor/Source/UVEditorTools/Private/Selection/UVEditorMeshSelectionMechanic.cpp b/Engine/Plugins/Editor/UVEditor/Source/UVEditorTools/Private/Selection/UVEditorMeshSelectionMechanic.cpp
--- a/Engine/Plugins/Editor/UVEditor/Source/UVEditorTools/Private/Selection/UVEditorMeshSelectionMechanic.cpp
+++ b/Engine/Plugins/Editor/UVEditor/Source/UVEditorTools/Private/Selection/UVEditorMeshSelectionMechanic.cpp
@@ -253,6 +253,90 @@ namespace UVEditorMeshSelectionMechanicLocals
 		}
 	};
 
+	struct FCircleRegion
+	{
+		FVector2d Center;
+		double RadiusSquared = 0;
+	};
+	template<>
+	struct TTreeBoxIntersectHelper<FCircleRegion>
+	{
+		typedef FVector2d PointType;
+		typedef FAxisAlignedBox2d BoxType;
+
+		static BoxType BuildBoxFromTree(const FDynamicMeshAABBTree3& Tree)
+		{
+			BoxType Box;
+			Box.Contain(XY(Tree.GetBoundingBox().Min));
+			Box.Contain(XY(Tree.GetBoundingBox().Max));
+			return Box;
+		}
+
+		static BoxType BuildBoxFromBox3d(const FAxisAlignedBox3d& Box3d)
+		{
+			return BoxType(XY(Box3d.Min), XY(Box3d.Max));
+		}
+
+		static bool RegionBoxTest(const FCircleRegion& Region, const FAxisAlignedBox2d& Box, bool& bFullyContained)
+		{
+			// Check for full containment first.
+			bool bIntersect = false;
+			bFullyContained = true;
+			for (int CornerIndex = 0; CornerIndex < 4; ++CornerIndex)
+			{
+				if (FVector2d::DistSquared(Box.GetCorner(CornerIndex), Region.Center) <= Region.RadiusSquared)
+				{
+					bIntersect = true;
+				}
+				else
+				{
+					bFullyContained = false;
+					// No need to test others because our intersection test below covers it.
+					break;
+				}
+			}
+
+			if (bIntersect)
+			{
+				return true;
+			}
+
+			// Now check for intersection. This is done the same way as FMath::SphereAABBIntersection, which
+			//  just checks our squared distance out from each box side, resulting in a sum if its past a corner.
+			double DistSquared = 0;
+			if (Region.Center.X < Box.Min.X)
+			{
+				DistSquared += FMath::Square(Region.Center.X - Box.Min.X);
+			}
+			else if (Region.Center.X > Box.Max.X)
+			{
+				DistSquared += FMath::Square(Region.Center.X - Box.Max.X);
+			}
+			if (Region.Center.Y < Box.Min.Y)
+			{
+				DistSquared += FMath::Square(Region.Center.Y - Box.Min.Y);
+			}
+			else if (Region.Center.Y > Box.Max.Y)
+			{
+				DistSquared += FMath::Square(Region.Center.Y - Box.Max.Y);
+			}
+
+			return DistSquared <= Region.RadiusSquared;
+		}
+	};
+
+	void AppendTriangleIDIfIntersectedCircle(const FDynamicMesh3& Mesh, const FCircleRegion& Circle, int TriangleID, TArray<int>& TriangleIDs)
+	{
+		FVector3d V1, V2, V3;
+		Mesh.GetTriVertices(TriangleID, V1, V2, V3);
+
+		if (UE::Geometry::DoesTriangleIntersectCircle2D(FVector2d(V1), FVector2d(V2), FVector2d(V3), Circle.Center, Circle.RadiusSquared))
+		{
+			TriangleIDs.Add(TriangleID);
+		}
+	}
+
+
 	// Returns indices, collected by the given functions, from triangles which are intersected by the given region,
 	// which can be either a rectangle (in which case the underlying mesh should have its vertices in the XY plane
 	// with a zero Z coordinate) or a frustum. 
@@ -733,6 +817,37 @@ void UUVEditorMeshSelectionMechanic::SetIsEnabled(bool bIsEnabledIn)
 	}
 }
 
+TArray<UE::Geometry::FUVToolSelection> UUVEditorMeshSelectionMechanic::GetAllCanonicalTrianglesInUnwrapRadius(
+	const FVector2d& UnwrapWorldHitPoint, double Radius) const
+{
+	using namespace UVEditorMeshSelectionMechanicLocals;
+
+	TArray<UE::Geometry::FUVToolSelection> Output;
+
+	FCircleRegion CircleRegion;
+	CircleRegion.Center = UnwrapWorldHitPoint;
+	CircleRegion.RadiusSquared = FMath::Square(FMath::Max(0, Radius));
+
+	for (int32 AssetID = 0; AssetID < Targets.Num(); ++AssetID)
+	{
+		const FDynamicMeshAABBTree3* Tree = GetMeshSpatial(AssetID, true).Get();
+		if (!ensure(Tree))
+		{
+			continue;
+		}
+
+		TArray<int32> TidsInRadius = FindAllIntersections(*Tree, CircleRegion, AppendTriangleID, AppendTriangleIDIfIntersectedCircle);
+		if (TidsInRadius.Num() > 0)
+		{
+			FUVToolSelection& Selection = Output.Emplace_GetRef();
+			Selection.Target = Targets[AssetID];
+			Selection.Type = FUVToolSelection::EType::Triangle;
+			Selection.SelectedIDs.Append(TidsInRadius);
+		}
+	}
+	return Output;
+}
+
 void UUVEditorMeshSelectionMechanic::SetShowHoveredElements(bool bShow)
 {
 	bShowHoveredElements = bShow;
@@ -825,7 +940,7 @@ void UUVEditorMeshSelectionMechanic::SetTargets(const TArray<TObjectPtr<UUVEdito
 
 }
 
-TSharedPtr<FDynamicMeshAABBTree3> UUVEditorMeshSelectionMechanic::GetMeshSpatial(int32 TargetId, bool bUseUnwrap)
+TSharedPtr<FDynamicMeshAABBTree3> UUVEditorMeshSelectionMechanic::GetMeshSpatial(int32 TargetId, bool bUseUnwrap) const
 {
 	if (bUseUnwrap)
 	{
@@ -1098,12 +1213,35 @@ void UUVEditorMeshSelectionMechanic::OnClicked(const FInputDeviceRay& ClickPos,
 	SelectionAPI->EndChangeAndEmitIfModified(true); // broadcast and emit
 }
 
-bool UUVEditorMeshSelectionMechanic::GetHitTid(const FInputDeviceRay& ClickPos, 
+// Usages of this should probably be replaced by RaycastCanonicals, but it doesn't seem worth the refactor.
+bool UUVEditorMeshSelectionMechanic::GetHitTid(const FInputDeviceRay& ClickPos,
 	int32& TidOut, int32& AssetIDOut, bool bUseUnwrap, int32* ExistingSelectionObjectIndexOut)
 {
-	auto RayCastSpatial = [this, &ClickPos, &TidOut, &AssetIDOut, &bUseUnwrap](int32 AssetID) {
+	FRaycastResult HitResult;
+	bool bSuccess = RaycastCanonicals(ClickPos.WorldRay, bUseUnwrap, true, HitResult);
+	AssetIDOut = HitResult.AssetID;
+	TidOut = HitResult.Tid;
+	if (ExistingSelectionObjectIndexOut)
+	{
+		*ExistingSelectionObjectIndexOut = IndexConstants::InvalidID;
+		if (bSuccess)
+		{
+			const TArray<FUVToolSelection>& Selections = SelectionAPI->GetSelections();
+			*ExistingSelectionObjectIndexOut = Selections.IndexOfByPredicate([this, AssetIDOut](const FUVToolSelection& Selection)
+			{
+				return Selection.Target == Targets[AssetIDOut];
+			});
+		}
+	}
+	return TidOut != IndexConstants::InvalidID;
+}
+
+bool UUVEditorMeshSelectionMechanic::RaycastCanonicals(const FRay& WorldRay, bool bRaycastIsForUnwrap,
+	bool bPreferSelected, FRaycastResult& HitOut) const
+{
+	auto RayCastSpatial = [this, &WorldRay, &HitOut, bRaycastIsForUnwrap](int32 AssetID) {
 		FTransform3d TargetTransform;
-		if (bUseUnwrap)
+		if (bRaycastIsForUnwrap)
 		{
 			TargetTransform = Targets[AssetID]->UnwrapPreview->PreviewMesh->GetTransform();
 		}
@@ -1113,13 +1251,14 @@ bool UUVEditorMeshSelectionMechanic::GetHitTid(const FInputDeviceRay& ClickPos,
 		}
 
 		FRay3d LocalRay(
-			(FVector3d)TargetTransform.InverseTransformPosition(ClickPos.WorldRay.Origin),
-			(FVector3d)TargetTransform.InverseTransformVector(ClickPos.WorldRay.Direction));
+			(FVector3d)TargetTransform.InverseTransformPosition(WorldRay.Origin),
+			(FVector3d)TargetTransform.InverseTransformVector(WorldRay.Direction));
 
 		double RayT = 0;
-		if (GetMeshSpatial(AssetID, bUseUnwrap)->FindNearestHitTriangle(LocalRay, RayT, TidOut))
+		if (GetMeshSpatial(AssetID, bRaycastIsForUnwrap)->FindNearestHitTriangle(LocalRay, RayT, HitOut.Tid))
 		{
-			AssetIDOut = AssetID;
+			HitOut.AssetID = AssetID;
+			HitOut.HitPosition = TargetTransform.TransformPosition(LocalRay.PointAt(RayT));
 			return true;
 		}
 		return false;
@@ -1129,28 +1268,23 @@ bool UUVEditorMeshSelectionMechanic::GetHitTid(const FInputDeviceRay& ClickPos,
 	TArray<bool> SpatialTriedFlags;
 	SpatialTriedFlags.SetNum(Targets.Num());
 	const TArray<FUVToolSelection>& Selections = SelectionAPI->GetSelections();
-	for (int32 SelectionIndex = 0; SelectionIndex < Selections.Num(); ++SelectionIndex)
+
+	if (bPreferSelected)
 	{
-		const FUVToolSelection& Selection = Selections[SelectionIndex];
-		if (ensure(Selection.Target.IsValid() && Selection.Target->AssetID < Targets.Num()))
+		for (int32 SelectionIndex = 0; SelectionIndex < Selections.Num(); ++SelectionIndex)
 		{
-			if (RayCastSpatial(Selection.Target->AssetID))
+			const FUVToolSelection& Selection = Selections[SelectionIndex];
+			if (ensure(Selection.Target.IsValid() && Selection.Target->AssetID < Targets.Num()))
 			{
-				if (ExistingSelectionObjectIndexOut)
+				if (RayCastSpatial(Selection.Target->AssetID))
 				{
-					*ExistingSelectionObjectIndexOut = SelectionIndex;
+					return true;
 				}
-				return true;
+				SpatialTriedFlags[Selection.Target->AssetID] = true;
 			}
-			SpatialTriedFlags[Selection.Target->AssetID] = true;
 		}
 	}
 
-	if (ExistingSelectionObjectIndexOut)
-	{
-		*ExistingSelectionObjectIndexOut = IndexConstants::InvalidID;
-	}
-
 	// Try raycasting the other meshes
 	for (int32 AssetID = 0; AssetID < Targets.Num(); ++AssetID)
 	{
diff --git a/Engine/Plugins/Editor/UVEditor/Source/UVEditorTools/Private/UVEditorBrushSelectTool.cpp b/Engine/Plugins/Editor/UVEditor/Source/UVEditorTools/Private/UVEditorBrushSelectTool.cpp
new file mode 100644
--- /dev/null
+++ b/Engine/Plugins/Editor/UVEditor/Source/UVEditorTools/Private/UVEditorBrushSelectTool.cpp
@@ -0,0 +1,611 @@
+// Copyright Epic Games, Inc. All Rights Reserved.
+
+#include "UVEditorBrushSelectTool.h"
+
+#include "BaseBehaviors/ClickDragBehavior.h"
+#include "BaseBehaviors/MouseHoverBehavior.h"
+#include "BaseGizmos/BrushStampIndicator.h"
+#include "BaseGizmos/GizmoMath.h"
+#include "ContextObjects/UVToolContextObjects.h"
+#include "ContextObjectStore.h"
+#include "InputCoreTypes.h" // EKey
+#include "InteractiveToolManager.h"
+#include "InteractiveGizmoManager.h"
+#include "Selection/UVEditorMeshSelectionMechanic.h"
+#include "Selections/MeshConnectedComponents.h"
+#include "Templates/Function.h"
+#include "UVEditorUXSettings.h"
+
+#define LOCTEXT_NAMESPACE "UUVEditorBrushSelectTool"
+
+namespace UVEditorBrushSelectToolLocals
+{
+	const FString BrushIndicatorGizmoType = TEXT("BrushIndicatorGizmoType");
+
+	const float MinUnwrapBrushRadius = 1.0 / (1 << 6);
+	const float MinLivePreviewBrushRadius = 100 / (1 << 5);
+
+	void UpdateSelections(TArray<UUVToolSelectionAPI::FUVToolSelection>& Selections, TSet<int32> Tids,
+		UUVEditorToolMeshInput* Target, bool bSubtract)
+	{
+		if (Tids.IsEmpty())
+		{
+			return;
+		}
+
+		int32 ExistingSelectionIndex = Selections.IndexOfByPredicate([Target](const UUVToolSelectionAPI::FUVToolSelection& ExistingSelection)
+		{
+			return ExistingSelection.Target == Target;
+		});
+
+		if (ExistingSelectionIndex >= 0)
+		{
+			UUVToolSelectionAPI::FUVToolSelection& ExistingSelection = Selections[ExistingSelectionIndex];
+			if (bSubtract)
+			{
+				ExistingSelection.SelectedIDs = ExistingSelection.SelectedIDs.Difference(Tids);
+				if (ExistingSelection.SelectedIDs.IsEmpty())
+				{
+					Selections.RemoveAtSwap(ExistingSelectionIndex);
+				}
+			}
+			else
+			{
+				ExistingSelection.SelectedIDs.Append(Tids);
+			}
+		}
+		// If no existing selection
+		else if (!bSubtract)
+		{
+			UUVToolSelectionAPI::FUVToolSelection& NewSelection = Selections.Emplace_GetRef();
+			NewSelection.Target = Target;
+			NewSelection.Type = UUVToolSelectionAPI::FUVToolSelection::EType::Triangle;
+			NewSelection.SelectedIDs = Tids;
+		}
+	}
+}
+
+void UUVEditorBrushSelectTool::Setup()
+{
+	using namespace UVEditorBrushSelectToolLocals;
+	using namespace UE::Geometry;
+
+	Super::Setup();
+
+	SetToolDisplayName(LOCTEXT("ToolName", "Brush Select Tool"));
+
+
+	Settings = NewObject<UUVEditorBrushSelectToolProperties>(this);
+	Settings->RestoreProperties(this);
+	AddToolPropertySource(Settings);
+
+	UContextObjectStore* ContextStore = GetToolManager()->GetContextObjectStore();
+	EmitChangeAPI = ContextStore->FindContext<UUVToolEmitChangeAPI>();
+	SelectionAPI = ContextStore->FindContext<UUVToolSelectionAPI>();
+	LivePreviewAPI = ContextStore->FindContext<UUVToolLivePreviewAPI>();
+
+	if (!ensure(SelectionAPI))
+	{
+		GetToolManager()->DeactivateTool(EToolSide::Mouse, EToolShutdownType::Completed);
+	}
+
+	SelectionMechanic = SelectionAPI->GetSelectionMechanic();
+
+	UUVToolSelectionAPI::FHighlightOptions HighlightOptions;
+	HighlightOptions.bAutoUpdateUnwrap = true;
+	HighlightOptions.bAutoUpdateApplied = true;
+	SelectionAPI->SetHighlightOptions(HighlightOptions);
+	SelectionAPI->SetHighlightVisible(true, true);
+
+	// Hook up to the input routers
+	ULocalClickDragInputBehavior* UnwrapClickDragBehavior = NewObject<ULocalClickDragInputBehavior>();
+	UnwrapClickDragBehavior->Initialize();
+	UnwrapClickDragBehavior->CanBeginClickDragFunc = [](const FInputDeviceRay&) 
+	{
+		// Accept all hits
+		return FInputRayHit(0.0);
+	};
+	UnwrapClickDragBehavior->OnClickPressFunc = [this, UnwrapClickDragBehavior](const FInputDeviceRay& DragPos)
+	{
+		SelectionAPI->BeginChange();
+		if (Settings->bClearSelectionOnEachDrag && !bShiftToggle && !bCtrlToggle)
+		{
+			ClearSelections();
+		}
+		bCurrentStrokeIsSubtracting = bCtrlToggle && !bShiftToggle;
+		// See TODO in header file
+		//bCurrentStrokeIsInverting = bCtrlToggle && bShiftToggle;
+
+		UnwrapClickDragBehavior->OnClickDragFunc(DragPos);
+		bHaveInteracted = true;
+	};
+	UnwrapClickDragBehavior->OnClickDragFunc = [this](const FInputDeviceRay& DragPos)
+	{
+		bool bIntersects;
+		FVector3d IntersectionPoint;
+		GizmoMath::RayPlaneIntersectionPoint(FVector::Zero(), FVector::ZAxisVector,
+			DragPos.WorldRay.Origin, DragPos.WorldRay.Direction, bIntersects, IntersectionPoint);
+		PendingUnwrapHits.Add(FVector2d(IntersectionPoint));
+		UpdateViewportStateFromHoverOrDragEvent(/*bFromUnwrap*/ true, /*bIsEndEvent*/ false, /*bDragging*/ true);
+		// Hover behaviors are terminated during drag, so update hover
+		UnwrapBrushIndicator->Update(Settings->UnwrapBrushRadius * FUVEditorUXSettings::UVMeshScalingFactor, 
+			IntersectionPoint, FVector::ZAxisVector, /*Falloff*/ 0.0f, /*Strength*/ 1.0f);
+	};
+	UnwrapClickDragBehavior->OnClickReleaseFunc = [this, UnwrapClickDragBehavior](const FInputDeviceRay&)
+	{
+		SelectionAPI->EndChangeAndEmitIfModified(true);
+		UnwrapClickDragBehavior->OnTerminateFunc();
+	};
+	UnwrapClickDragBehavior->OnTerminateFunc = [this]()
+	{
+		PendingUnwrapHits.Reset();
+		UpdateViewportStateFromHoverOrDragEvent(/*bFromUnwrap*/ true, /*bIsEndEvent*/ true, /*bDragging*/ true);
+	};
+	AddInputBehavior(UnwrapClickDragBehavior);
+
+	ULocalMouseHoverBehavior* UnwrapHoverBehavior = NewObject<ULocalMouseHoverBehavior>();
+	UnwrapHoverBehavior->Initialize();
+	UnwrapHoverBehavior->BeginHitTestFunc = [](const FInputDeviceRay& DragPos) {
+		// Accept all hits
+		return FInputRayHit(0.0);
+	};
+	UnwrapHoverBehavior->OnBeginHoverFunc = [this, UnwrapHoverBehavior](const FInputDeviceRay& DragPos)
+	{
+		UnwrapHoverBehavior->OnUpdateHover(DragPos);
+	};
+	UnwrapHoverBehavior->OnUpdateHoverFunc = [this](const FInputDeviceRay& DragPos)
+	{
+		bool bIntersects;
+		FVector3d IntersectionPoint;
+		GizmoMath::RayPlaneIntersectionPoint(FVector::Zero(), FVector::ZAxisVector,
+			DragPos.WorldRay.Origin, DragPos.WorldRay.Direction, bIntersects, IntersectionPoint);
+		UpdateViewportStateFromHoverOrDragEvent(/*bFromUnwrap*/ true, /*bIsEndEvent*/ false, /*bDragging*/ false);
+		UnwrapBrushIndicator->Update(Settings->UnwrapBrushRadius * FUVEditorUXSettings::UVMeshScalingFactor, 
+			IntersectionPoint, FVector::ZAxisVector, /*Falloff*/ 0.0f, /*Strength*/ 1.0f);
+
+		// Should always be true, actually, since we're shooting a ray down into the plane. But might
+		//  as well return this bool so that if something is incorrectly set up, we end hover.
+		return bIntersects;
+	};
+	UnwrapHoverBehavior->OnEndHoverFunc = [this]()
+	{
+		UpdateViewportStateFromHoverOrDragEvent(/*bFromUnwrap*/ true, /*bIsEndEvent*/ true, /*bDragging*/ false);
+	};
+	UnwrapHoverBehavior->Modifiers.RegisterModifier(ShiftModifierID, FInputDeviceState::IsShiftKeyDown);
+	UnwrapHoverBehavior->Modifiers.RegisterModifier(CtrlModifierID, FInputDeviceState::IsCtrlKeyDown);
+	auto UpdateModifiers = [this](int ModifierID, bool bIsOn) 
+	{
+		if (ModifierID == ShiftModifierID)
+		{
+			bShiftToggle = bIsOn;
+		}
+		else if (ModifierID == CtrlModifierID)
+		{
+			bCtrlToggle = bIsOn;
+		}
+	};
+	UnwrapHoverBehavior->OnUpdateModifierStateFunc = UpdateModifiers;
+	AddInputBehavior(UnwrapHoverBehavior);
+
+	ULocalClickDragInputBehavior* LivePreviewClickDragBehavior = NewObject<ULocalClickDragInputBehavior>();
+	LivePreviewClickDragBehavior->Initialize();
+	LivePreviewClickDragBehavior->CanBeginClickDragFunc = [this](const FInputDeviceRay& DragPos)
+	{
+		UUVEditorMeshSelectionMechanic::FRaycastResult HitResult;
+		if (SelectionMechanic.IsValid() && SelectionMechanic->RaycastCanonicals(DragPos.WorldRay, false, false, HitResult))
+		{
+			return FInputRayHit(DragPos.WorldRay.GetParameter(HitResult.HitPosition));
+		}
+		return FInputRayHit();
+	};
+	LivePreviewClickDragBehavior->OnClickPressFunc = [this, LivePreviewClickDragBehavior](const FInputDeviceRay& DragPos)
+	{
+		SelectionAPI->BeginChange();
+		if (Settings->bClearSelectionOnEachDrag && !bShiftToggle && !bCtrlToggle)
+		{
+			ClearSelections();
+		}
+		bCurrentStrokeIsSubtracting = bCtrlToggle && !bShiftToggle;
+		// See TODO in header
+		//bCurrentStrokeIsInverting = bCtrlToggle && bShiftToggle;
+
+		LivePreviewClickDragBehavior->OnClickDragFunc(DragPos);
+		bHaveInteracted = true;
+	};
+	LivePreviewClickDragBehavior->OnClickDragFunc = [this](const FInputDeviceRay& DragPos)
+	{
+		UUVEditorMeshSelectionMechanic::FRaycastResult HitResult;
+		if (SelectionMechanic.IsValid() && SelectionMechanic->RaycastCanonicals(DragPos.WorldRay, false, false, HitResult))
+		{
+			PendingLivePreviewHits.Add(HitResult);
+
+			FVector3d Normal = Targets[HitResult.AssetID]->AppliedCanonical->GetTriNormal(HitResult.Tid);
+			UpdateViewportStateFromHoverOrDragEvent(/*bFromUnwrap*/ false, /*bIsEndEvent*/ false, /*bDragging*/ true);
+			// Hover behaviors are terminated during drag, so update hover
+			LivePreviewBrushIndicator->Update(Settings->LivePreviewBrushRadius, 
+				HitResult.HitPosition, Normal, /*Falloff*/ 0.0f, /*Strength*/ 1.0f);
+		}
+	};
+	LivePreviewClickDragBehavior->OnClickReleaseFunc = [this, LivePreviewClickDragBehavior](const FInputDeviceRay& DragPos)
+	{
+		SelectionAPI->EndChangeAndEmitIfModified(true);
+		LivePreviewClickDragBehavior->OnTerminateFunc();
+	};
+	LivePreviewClickDragBehavior->OnTerminateFunc = [this]()
+	{
+		PendingLivePreviewHits.Reset();
+		UpdateViewportStateFromHoverOrDragEvent(/*bFromUnwrap*/ false, /*bIsEndEvent*/ true, /*bDragging*/ true);
+	};
+
+	ULocalMouseHoverBehavior* LivePreviewHoverBehavior = NewObject<ULocalMouseHoverBehavior>();
+	LivePreviewHoverBehavior->Initialize();
+	LivePreviewHoverBehavior->BeginHitTestFunc = [this](const FInputDeviceRay& DragPos)
+	{
+		UUVEditorMeshSelectionMechanic::FRaycastResult HitResult;
+		if (SelectionMechanic.IsValid() && SelectionMechanic->RaycastCanonicals(DragPos.WorldRay, false, false, HitResult))
+		{
+			return FInputRayHit(DragPos.WorldRay.GetParameter(HitResult.HitPosition));
+		}
+		return FInputRayHit();
+	};
+	LivePreviewHoverBehavior->OnBeginHoverFunc = [this, LivePreviewHoverBehavior](const FInputDeviceRay& DragPos)
+	{
+		LivePreviewHoverBehavior->OnUpdateHover(DragPos);
+	};
+	LivePreviewHoverBehavior->OnUpdateHoverFunc = [this](const FInputDeviceRay& DragPos)
+	{
+		UUVEditorMeshSelectionMechanic::FRaycastResult HitResult;
+		if (SelectionMechanic.IsValid() && SelectionMechanic->RaycastCanonicals(DragPos.WorldRay, false, false, HitResult))
+		{
+			FVector3d Normal = Targets[HitResult.AssetID]->AppliedCanonical->GetTriNormal(HitResult.Tid);
+			UpdateViewportStateFromHoverOrDragEvent(/*bFromUnwrap*/ false, /*bIsEndEvent*/ false, /*bDragging*/ false);
+			LivePreviewBrushIndicator->Update(Settings->LivePreviewBrushRadius, 
+				HitResult.HitPosition, Normal, /*Falloff*/ 0.0f, /*Strength*/ 1.0f);
+			return true;
+		}
+		return false;
+	};
+	LivePreviewHoverBehavior->OnEndHoverFunc = [this]()
+	{
+		UpdateViewportStateFromHoverOrDragEvent(/*bFromUnwrap*/ false, /*bIsEndEvent*/ true, /*bDragging*/ false);
+	};
+	LivePreviewHoverBehavior->Modifiers.RegisterModifier(ShiftModifierID, FInputDeviceState::IsShiftKeyDown);
+	LivePreviewHoverBehavior->Modifiers.RegisterModifier(CtrlModifierID, FInputDeviceState::IsCtrlKeyDown);
+	LivePreviewHoverBehavior->OnUpdateModifierStateFunc = UpdateModifiers;
+	
+	if (LivePreviewAPI)
+	{
+		LivePreviewInputRouter = LivePreviewAPI->GetLivePreviewInputRouter();
+		LivePreviewBehaviorSet = NewObject<UInputBehaviorSet>();
+		LivePreviewBehaviorSource = NewObject<ULocalInputBehaviorSource>();
+		LivePreviewBehaviorSource->GetInputBehaviorsFunc = [this]() { return LivePreviewBehaviorSet; };
+		LivePreviewBehaviorSet->Add(LivePreviewClickDragBehavior, this);
+		LivePreviewBehaviorSet->Add(LivePreviewHoverBehavior, this);
+
+		if (ensure(LivePreviewInputRouter.IsValid()))
+		{
+			LivePreviewInputRouter->RegisterSource(LivePreviewBehaviorSource);
+		}
+
+		// register and spawn brush indicator gizmo for live preview
+		UInteractiveGizmoManager* LivePreviewGizmoManager = LivePreviewAPI->GetGizmoManager();
+		if (ensure(LivePreviewGizmoManager))
+		{
+			LivePreviewGizmoManager->RegisterGizmoType(BrushIndicatorGizmoType, NewObject<UBrushStampIndicatorBuilder>());
+			LivePreviewBrushIndicator = LivePreviewGizmoManager->CreateGizmo<UBrushStampIndicator>(
+				BrushIndicatorGizmoType, FString(), this);
+		}
+	}
+	
+	// register and spawn brush indicator gizmo for unwrap
+	GetToolManager()->GetPairedGizmoManager()->RegisterGizmoType(BrushIndicatorGizmoType, NewObject<UBrushStampIndicatorBuilder>());
+	UnwrapBrushIndicator = GetToolManager()->GetPairedGizmoManager()->CreateGizmo<UBrushStampIndicator>(BrushIndicatorGizmoType, FString(), this);
+	
+	GetToolManager()->DisplayMessage(LOCTEXT("StatusBarMessage",
+		"Shift adds to selection, Ctrl subtracts, both together toggle. [ and ] change brush size."),
+		EToolMessageLevel::UserNotification);
+}
+
+void UUVEditorBrushSelectTool::Shutdown(EToolShutdownType ShutdownType)
+{
+	using namespace UVEditorBrushSelectToolLocals;
+
+	if (LivePreviewInputRouter.IsValid())
+	{
+		// TODO: Arguably the live preview input router should do this for us before Shutdown(), but
+		// we don't currently have support for that...
+		LivePreviewInputRouter->ForceTerminateSource(LivePreviewBehaviorSource);
+		LivePreviewInputRouter->DeregisterSource(LivePreviewBehaviorSource);
+	}
+
+	GetToolManager()->GetPairedGizmoManager()->DestroyGizmo(UnwrapBrushIndicator);
+	GetToolManager()->GetPairedGizmoManager()->DeregisterGizmoType(BrushIndicatorGizmoType);
+	UnwrapBrushIndicator = nullptr;
+
+	if (LivePreviewAPI)
+	{
+		if (UInteractiveGizmoManager* LivePreviewGizmoManager = LivePreviewAPI->GetGizmoManager())
+		{
+			LivePreviewGizmoManager->DestroyGizmo(LivePreviewBrushIndicator);
+			LivePreviewGizmoManager->DeregisterGizmoType(BrushIndicatorGizmoType);
+		}
+	}
+	LivePreviewBrushIndicator = nullptr;
+
+	ProcessPendingUnwrapHits();
+	ProcessPendingLivePreviewHits();
+	SelectionAPI->EndChangeAndEmitIfModified(true);
+
+	LivePreviewBehaviorSet->RemoveAll();
+	LivePreviewBehaviorSet = nullptr;
+	LivePreviewBehaviorSource = nullptr;
+
+	Settings->SaveProperties(this);
+	Settings = nullptr;
+
+	SelectionAPI = nullptr;
+	EmitChangeAPI = nullptr;
+	LivePreviewAPI = nullptr;
+
+	Super::Shutdown(ShutdownType);
+}
+
+void UUVEditorBrushSelectTool::Render(IToolsContextRenderAPI* RenderAPI)
+{
+}
+
+void UUVEditorBrushSelectTool::OnTick(float DeltaTime)
+{
+	ProcessPendingUnwrapHits();
+	ProcessPendingLivePreviewHits();
+}
+
+void UUVEditorBrushSelectTool::ClearSelections(bool bBroadcastAndEmit)
+{
+	if (bBroadcastAndEmit)
+	{
+		EmitChangeAPI->BeginUndoTransaction(LOCTEXT("ClearSelectionTransaction", "Clear Selection"));
+	}
+	SelectionAPI->ClearUnsetElementAppliedMeshSelections(bBroadcastAndEmit, bBroadcastAndEmit);
+	SelectionAPI->ClearSelections(bBroadcastAndEmit, bBroadcastAndEmit);
+	if (bBroadcastAndEmit)
+	{
+		EmitChangeAPI->EndUndoTransaction();
+	}
+}
+
+void UUVEditorBrushSelectTool::ProcessPendingUnwrapHits()
+{
+	using namespace UE::Geometry;
+	using namespace UVEditorBrushSelectToolLocals;
+
+	if (!SelectionMechanic.IsValid() || PendingUnwrapHits.IsEmpty())
+	{
+		return;
+	}
+
+	TArray<FUVToolSelection> SelectionsToSet = SelectionAPI->GetSelections();
+
+	for (const FVector2d& WorldHitPoint : PendingUnwrapHits)
+	{
+		for (FUVToolSelection& Selection : SelectionMechanic->GetAllCanonicalTrianglesInUnwrapRadius(
+			WorldHitPoint, Settings->UnwrapBrushRadius * FUVEditorUXSettings::UVMeshScalingFactor))
+		{
+			if (Settings->bExpandToIslands)
+			{
+				TArray<int32> ExpandedTids;
+				FMeshConnectedComponents::GrowToConnectedTriangles(
+					Selection.Target->UnwrapCanonical.Get(), 
+					Selection.SelectedIDs.Array(), ExpandedTids, &TempROIBuffer);
+				Selection.SelectedIDs.Append(ExpandedTids);
+			}
+
+			UpdateSelections(SelectionsToSet, Selection.SelectedIDs, Selection.Target.Get(), 
+				bCurrentStrokeIsSubtracting);
+		}
+	}
+
+	SelectionAPI->SetSelections(SelectionsToSet, false, false);
+
+	PendingUnwrapHits.Reset();
+}
+
+void UUVEditorBrushSelectTool::ProcessPendingLivePreviewHits()
+{
+	using namespace UE::Geometry;
+	using namespace UVEditorBrushSelectToolLocals;
+
+	if (!SelectionMechanic.IsValid() || PendingLivePreviewHits.IsEmpty())
+	{
+		return;
+	}
+
+	TArray<FUVToolSelection> SelectionsToSet = SelectionAPI->GetSelections();
+	TArray<FUVToolSelection> UnsetSelectionsToSet = SelectionAPI->GetUnsetElementAppliedMeshSelections();
+
+	for (const UUVEditorMeshSelectionMechanic::FRaycastResult& Hit : PendingLivePreviewHits)
+	{
+		const FDynamicMesh3* Mesh = Targets[Hit.AssetID]->AppliedCanonical.Get();
+		const FDynamicMeshUVOverlay* UVOverlay = Mesh->Attributes()->GetUVLayer(Targets[Hit.AssetID]->UVLayerIndex);
+
+		if (Mesh->IsTriangle(Hit.Tid))
+		{
+			double RadiusSquared = Settings->LivePreviewBrushRadius * Settings->LivePreviewBrushRadius;
+
+			TFunction<bool(int32, int32)> GrowPredicate = [Mesh, &Hit, RadiusSquared](int t1, int t2)
+			{
+				return (Mesh->GetTriCentroid(t2) - Hit.HitPosition).SquaredLength() <= RadiusSquared;
+			};
+			if (Settings->bExpandToIslands)
+			{
+				GrowPredicate = [Mesh, UVOverlay, &Hit, RadiusSquared](int t1, int t2)
+				{
+					int32 Edge = Mesh->FindEdgeFromTriPair(t1, t2);
+					return !UVOverlay->IsSeamEdge(Edge) 
+						|| (Mesh->GetTriCentroid(t2) - Hit.HitPosition).SquaredLength() <= RadiusSquared;
+				};
+			}
+
+			TArray<int32> StartSet { Hit.Tid };
+			TSet<int32> GrownAppliedTids;
+			FMeshConnectedComponents::GrowToConnectedTriangles(Mesh, StartSet, GrownAppliedTids, &TempROIBuffer, GrowPredicate);
+
+			const FDynamicMesh3* UnwrapMesh = Targets[Hit.AssetID]->UnwrapCanonical.Get();
+			TSet<int32> UnwrapTids;
+			TSet<int32> UnsetTids;
+			for (int32 Tid : GrownAppliedTids)
+			{
+				if (UnwrapMesh->IsTriangle(Tid))
+				{
+					UnwrapTids.Add(Tid);
+				}
+				else
+				{
+					UnsetTids.Add(Tid);
+				}
+			}
+
+			UpdateSelections(SelectionsToSet, UnwrapTids, Targets[Hit.AssetID], 
+				bCurrentStrokeIsSubtracting);
+			UpdateSelections(UnsetSelectionsToSet, UnsetTids, Targets[Hit.AssetID], 
+				bCurrentStrokeIsSubtracting);
+		}
+	}
+
+	SelectionAPI->SetSelections(SelectionsToSet, false, false);
+	SelectionAPI->SetUnsetElementAppliedMeshSelections(UnsetSelectionsToSet, false, false);
+
+	PendingLivePreviewHits.Reset();
+}
+
+void UUVEditorBrushSelectTool::RegisterActions(FInteractiveToolActionSet& ActionSet)
+{
+	int32 ActionID = (int32)EStandardToolActions::BaseClientDefinedActionID + 1;
+	
+	// TODO: We could/should use EStandardToolActions::IncreaseBrushSize and EStandardToolActions::DecreaseBrushSize
+	//  as the action IDs here so that the same brush resizing hotkeys can act everywhere. The problem is that
+	//  there is currently a bug where if we don't have at least one action that is BaseClientDefinedActionID or
+	//  above, we will get a crash in our tool, see UE-221911.
+	ActionSet.RegisterAction(this, ActionID++,
+		TEXT("IncreaseRadius"),
+		LOCTEXT("IncreaseRadius", "Increase Radius"),
+		LOCTEXT("IncreaseRadiusTooltip", "Increase Brush Radius"),
+		EModifierKey::None, EKeys::RightBracket,
+		[this]() { IncreaseBrushRadiusAction(); });
+
+	ActionSet.RegisterAction(this, ActionID++,
+		TEXT("DecreaseRadius"),
+		LOCTEXT("DecreaseRadius", "Decrease Radius"),
+		LOCTEXT("DecreaseRadiusTooltip", "Decrease Brush Radius"),
+		EModifierKey::None, EKeys::LeftBracket,
+		[this]() { DecreaseBrushRadiusAction(); });
+}
+
+// We track which viewport (unwrap or live preview) we're hovering so that we can display the relevant brush
+//  and so that we can route hotkey presses to update the correct brush. This turns out to be a little tricky
+//  because we don't actually get an end hover event when we go from one viewport to the other, and because we
+//  DO get an end hover event AFTER getting a drag start event...
+void UUVEditorBrushSelectTool::UpdateViewportStateFromHoverOrDragEvent(bool bFromUnwrap, bool bIsEndEvent, bool bDragging)
+{
+	if (!bIsEndEvent)
+	{
+		// If we're getting a non-end event, then we definitely know that we're over one viewport
+		//  and not the other.
+		bHoveringUnwrap = bFromUnwrap;
+		bHoveringLivePreview = !bFromUnwrap;
+		if (bDragging)
+		{
+			bDraggingUnwrap = bFromUnwrap;
+			bDraggingLivePreview = !bFromUnwrap;
+		}
+	}
+	// If this is an end event, then it's trickier. If we're ending a drag, we can safely say
+	//  that we're not over that viewport (if we're still over it, our next hover begin event
+	//  will let us know). However if it's a hover end event, that frequently comes after a
+	//  begin drag event, so we don't want to mark ourselves as not being over that viewport.
+	else if (bDragging)
+	{
+		if (bFromUnwrap)
+		{
+			bDraggingUnwrap = false;
+			bHoveringUnwrap = false;
+		}
+		else
+		{
+			bDraggingLivePreview = false;
+			bHoveringLivePreview = false;
+		}
+	}
+	else // end hover event
+	{
+		if (bFromUnwrap)
+		{
+			bHoveringUnwrap = bDraggingUnwrap;
+		}
+		else
+		{
+			bHoveringLivePreview = bDraggingLivePreview;
+		}
+	}
+
+	// Update our brushes
+	if (UnwrapBrushIndicator)
+	{
+		UnwrapBrushIndicator->bVisible = bHoveringUnwrap;
+	}
+	if (LivePreviewBrushIndicator)
+	{
+		LivePreviewBrushIndicator->bVisible = bHoveringLivePreview;
+	}
+}
+
+// TODO: Probably will want to rework the increase/decrease actions. Might want the live preview radius in
+//  particular to be responsive to the bounds of the mesh, like most of our brush tools.
+void UUVEditorBrushSelectTool::IncreaseBrushRadiusAction()
+{
+	using namespace UVEditorBrushSelectToolLocals;
+
+	if (bHoveringLivePreview)
+	{
+		Settings->LivePreviewBrushRadius = Settings->LivePreviewBrushRadius <= 0 ? MinLivePreviewBrushRadius
+			: Settings->LivePreviewBrushRadius * 2;
+	}
+	if (bHoveringUnwrap)
+	{
+		Settings->UnwrapBrushRadius = Settings->UnwrapBrushRadius <= 0 ? MinUnwrapBrushRadius
+			: Settings->UnwrapBrushRadius * 2;
+	}
+}
+void UUVEditorBrushSelectTool::DecreaseBrushRadiusAction()
+{
+	using namespace UVEditorBrushSelectToolLocals;
+
+	if (bHoveringLivePreview)
+	{
+		Settings->LivePreviewBrushRadius = Settings->LivePreviewBrushRadius < MinLivePreviewBrushRadius + KINDA_SMALL_NUMBER ? 0.0
+			: Settings->LivePreviewBrushRadius / 2;
+	}
+	if (bHoveringUnwrap)
+	{
+		Settings->UnwrapBrushRadius = Settings->UnwrapBrushRadius < MinUnwrapBrushRadius + KINDA_SMALL_NUMBER ? 0.0
+			: Settings->UnwrapBrushRadius / 2;
+	}
+}
+
+bool UUVEditorBrushSelectTool::CanCurrentlyNestedCancel()
+{
+	return bHaveInteracted && (SelectionAPI->HaveSelections() 
+		|| SelectionAPI->HaveUnsetElementAppliedMeshSelections());
+}
+
+bool UUVEditorBrushSelectTool::ExecuteNestedCancelCommand()
+{
+	if (CanCurrentlyNestedCancel())
+	{
+		ClearSelections(true);
+		return true;
+	}
+	return false;
+}
+
+#undef LOCTEXT_NAMESPACE
diff --git a/Engine/Plugins/Editor/UVEditor/Source/UVEditorTools/Private/UVEditorToolBase.cpp b/Engine/Plugins/Editor/UVEditor/Source/UVEditorTools/Private/UVEditorToolBase.cpp
new file mode 100644
--- /dev/null
+++ b/Engine/Plugins/Editor/UVEditor/Source/UVEditorTools/Private/UVEditorToolBase.cpp
@@ -0,0 +1,32 @@
+// Copyright Epic Games, Inc. All Rights Reserved.
+
+#include "UVEditorToolBase.h"
+
+#include "InteractiveTool.h"
+#include "InteractiveToolManager.h"
+#include "ToolContextInterfaces.h" // FToolBuilderState
+
+void UGenericUVEditorToolBuilder::Initialize(TArray<TObjectPtr<UUVEditorToolMeshInput>>& TargetsIn, TSubclassOf<UInteractiveTool> ToolClassIn)
+{
+	Targets = &TargetsIn;
+	if (ensure(ToolClassIn && ToolClassIn->ImplementsInterface(UUVEditorGenericBuildableTool::StaticClass())))
+	{
+		ToolClass = ToolClassIn;
+	}
+}
+
+bool UGenericUVEditorToolBuilder::CanBuildTool(const FToolBuilderState& SceneState) const
+{
+	return ToolClass && Targets && Targets->Num() > 0;
+}
+
+UInteractiveTool* UGenericUVEditorToolBuilder::BuildTool(const FToolBuilderState& SceneState) const
+{
+	UInteractiveTool* NewTool = NewObject<UInteractiveTool>(SceneState.ToolManager, ToolClass.Get());
+	IUVEditorGenericBuildableTool* CastTool = Cast<IUVEditorGenericBuildableTool>(NewTool);
+	if (ensure(CastTool))
+	{
+		CastTool->SetTargets(*Targets);
+	}
+	return NewTool;
+}
\ No newline at end of file
diff --git a/Engine/Plugins/Editor/UVEditor/Source/UVEditorTools/Public/Selection/UVEditorMeshSelectionMechanic.h b/Engine/Plugins/Editor/UVEditor/Source/UVEditorTools/Public/Selection/UVEditorMeshSelectionMechanic.h
--- a/Engine/Plugins/Editor/UVEditor/Source/UVEditorTools/Public/Selection/UVEditorMeshSelectionMechanic.h
+++ b/Engine/Plugins/Editor/UVEditor/Source/UVEditorTools/Public/Selection/UVEditorMeshSelectionMechanic.h
@@ -54,6 +54,23 @@ public:
 	void SetIsEnabled(bool bIsEnabled);
 	bool IsEnabled() { return bIsEnabled; };
 
+	struct FRaycastResult
+	{
+		int32 AssetID = IndexConstants::InvalidID;
+		int32 Tid = IndexConstants::InvalidID;
+		FVector3d HitPosition;
+	};
+	/**
+	 * This is a helper method that doesn't get used in normal selection mechanic operation, but
+	 *  can be used by clients if they need to raycast canonical meshes, since the mechanic already
+	 *  keeps aabb trees for them. The mechanic does not need to be enabled to use this (and typically
+	 *  will not be, if a tool is doing its own raycasting).
+	 */
+	bool RaycastCanonicals(const FRay& WorldRay, bool bRaycastIsForUnwrap,
+		bool bPreferSelected, FRaycastResult& HitOut) const;
+
+	TArray<FUVToolSelection> GetAllCanonicalTrianglesInUnwrapRadius(const FVector2d& UnwrapWorldHitPoint, double Radius) const;
+
 	void SetShowHoveredElements(bool bShow);
 
 	using ESelectionMode = UUVToolSelectionAPI::EUVEditorSelectionMode;
@@ -145,7 +162,7 @@ protected:
 	TWeakObjectPtr<UPointSetComponent> LivePreviewHoverPointSet = nullptr;
 
 	// Should be the same as the mode-level targets array, indexed by AssetID
-	TSharedPtr<FDynamicMeshAABBTree3> GetMeshSpatial(int32 TargetId, bool bUseUnwrap);
+	TSharedPtr<FDynamicMeshAABBTree3> GetMeshSpatial(int32 TargetId, bool bUseUnwrap) const;
 	TArray<TObjectPtr<UUVEditorToolMeshInput>> Targets;
 	TArray<TSharedPtr<FDynamicMeshAABBTree3>> UnwrapMeshSpatials; // 1:1 with Targets
 	TArray<TSharedPtr<FDynamicMeshAABBTree3>> AppliedMeshSpatials; // 1:1 with Targets
diff --git a/Engine/Plugins/Editor/UVEditor/Source/UVEditorTools/Public/Selection/UVToolSelectionAPI.h b/Engine/Plugins/Editor/UVEditor/Source/UVEditorTools/Public/Selection/UVToolSelectionAPI.h
--- a/Engine/Plugins/Editor/UVEditor/Source/UVEditorTools/Public/Selection/UVToolSelectionAPI.h
+++ b/Engine/Plugins/Editor/UVEditor/Source/UVEditorTools/Public/Selection/UVToolSelectionAPI.h
@@ -132,6 +132,12 @@ public:
 
 	// Selection mechanic controls:
 
+	/**
+	 * Most clients shouldn't need to touch the selection mechanic directly, but it can
+	 * be useful for certain utility functions that exist on the mechanic.
+	 */
+	const UUVEditorMeshSelectionMechanic* GetSelectionMechanic() const { return SelectionMechanic; }
+
 	void SetSelectionMechanicEnabled(bool bEnabled);
 
 	struct FSelectionMechanicOptions
diff --git a/Engine/Plugins/Editor/UVEditor/Source/UVEditorTools/Public/UVEditorBrushSelectTool.h b/Engine/Plugins/Editor/UVEditor/Source/UVEditorTools/Public/UVEditorBrushSelectTool.h
new file mode 100644
--- /dev/null
+++ b/Engine/Plugins/Editor/UVEditor/Source/UVEditorTools/Public/UVEditorBrushSelectTool.h
@@ -0,0 +1,161 @@
+// Copyright Epic Games, Inc. All Rights Reserved.
+
+#pragma once
+
+#include "Containers/Array.h"
+#include "UVEditorToolBase.h" // IUVEditorGenericBuildableTool
+#include "InteractiveTool.h"
+#include "Selection/UVEditorMeshSelectionMechanic.h"
+#include "Selection/UVToolSelectionAPI.h" // IUVToolSupportsSelection
+#include "InteractiveToolQueryInterfaces.h" // IInteractiveToolNestedAcceptCancelAPI
+
+#include "UVEditorBrushSelectTool.generated.h"
+
+class UBrushStampIndicator;
+class ULocalMouseHoverBehavior;
+class ULocalClickDragInputBehavior;
+
+UCLASS()
+class UUVEditorBrushSelectToolProperties : public UInteractiveToolPropertySet
+{
+	GENERATED_BODY()
+
+public:
+	/** 
+	 * When true, each drag will clear the existing selection if neither Shift (add to selection) nor
+	 *  Ctrl (remove from selection) nor both (toggle selection) are pressed. When false, selection
+	 *  will not clear, requiring manual removal from selection using Ctrl.
+	 */
+	UPROPERTY(EditAnywhere, Category = Options)
+	bool bClearSelectionOnEachDrag = true;
+	
+	/** When true, brush selects whole UV islands instead of individual triangles. */
+	UPROPERTY(EditAnywhere, Category = Options)
+	bool bExpandToIslands = false;
+	
+	//~ The radius defaults are set up so that after doubling via hotkeys some number of times (bitshift
+	//~  amount), we get roughly 1.0 and 100.
+
+	/** Radius of the brush in the 2D UV unwrap view. */
+	UPROPERTY(EditAnywhere, Category = Options, meta = (DisplayName = "Unwrap Radius", 
+		ClampMin = "0", ClampMax = "1000000", UIMax = "1"))
+	float UnwrapBrushRadius = 1.0f / (1 << 5);
+
+	/** Radius of the brush in the 3D live preview view. */
+	UPROPERTY(EditAnywhere, Category = Options, meta = (DisplayName = "Preview Radius", 
+		ClampMin = "0", ClampMax = "1000000", UIMax = "100"))
+	float LivePreviewBrushRadius = 100.0f / (1 << 4);
+};
+
+//~ We could potentially fold this tool's functionality into the UVSelectTool as a sub activity
+//~  of some kind.
+/**
+ * The brush select tool allows for brush selection of triangles on the unwrap or the live
+ *  preview. 
+ */
+UCLASS()
+class UVEDITORTOOLS_API UUVEditorBrushSelectTool : public UInteractiveTool
+	, public IInteractiveToolNestedAcceptCancelAPI
+	, public IUVToolSupportsSelection
+	, public IUVEditorGenericBuildableTool
+{
+	GENERATED_BODY()
+
+public:
+	
+	// UInteractiveTool
+	virtual void Setup() override;
+	virtual void Shutdown(EToolShutdownType ShutdownType) override;
+	virtual void Render(IToolsContextRenderAPI* RenderAPI) override;
+	virtual void OnTick(float DeltaTime) override;
+	virtual bool HasCancel() const override { return false; }
+	virtual bool HasAccept() const override { return false; }
+	virtual void RegisterActions(FInteractiveToolActionSet& ActionSet) override;
+
+	// IInteractiveToolNestedAcceptCancelAPI
+	virtual bool SupportsNestedCancelCommand() override { return true; }
+	virtual bool CanCurrentlyNestedCancel() override;
+	virtual bool ExecuteNestedCancelCommand() override;
+
+	// IUVEditorGenericBuildableTool
+	virtual void SetTargets(const TArray<TObjectPtr<UUVEditorToolMeshInput>>& TargetsIn)
+	{
+		Targets = TargetsIn;
+	}
+
+	virtual bool SupportsUnsetElementAppliedMeshSelections() const { return true; }
+
+private:
+	UPROPERTY()
+	TArray<TObjectPtr<UUVEditorToolMeshInput>> Targets;
+
+	UPROPERTY()
+	TObjectPtr<UUVToolEmitChangeAPI> EmitChangeAPI = nullptr;
+
+	UPROPERTY()
+	TObjectPtr<UUVToolSelectionAPI> SelectionAPI = nullptr;
+
+	TWeakObjectPtr<const UUVEditorMeshSelectionMechanic> SelectionMechanic = nullptr;
+
+	UPROPERTY()
+	TObjectPtr<UUVToolLivePreviewAPI> LivePreviewAPI = nullptr;
+
+	TWeakObjectPtr<UInputRouter> LivePreviewInputRouter = nullptr;
+	UPROPERTY()
+	TObjectPtr<UInputBehaviorSet> LivePreviewBehaviorSet = nullptr;
+	UPROPERTY()
+	TObjectPtr<ULocalInputBehaviorSource> LivePreviewBehaviorSource = nullptr;
+
+	UPROPERTY()
+	TObjectPtr<UUVEditorBrushSelectToolProperties> Settings = nullptr;
+
+	UPROPERTY()
+	TObjectPtr<UBrushStampIndicator> UnwrapBrushIndicator = nullptr;
+
+	UPROPERTY()
+	TObjectPtr<UBrushStampIndicator> LivePreviewBrushIndicator = nullptr;
+
+	TArray<UUVEditorMeshSelectionMechanic::FRaycastResult> PendingLivePreviewHits;
+	// For unwrap, we don't necessarily need to hit the mesh to be able to select things
+	//  with the brush- we just need the locations of all the hits.
+	TArray<FVector2d> PendingUnwrapHits;
+
+	// Updates internal structures in addition to the actual selection api
+	void ClearSelections(bool bBroadcastAndEmit = false);
+
+	void ProcessPendingUnwrapHits();
+	void ProcessPendingLivePreviewHits();
+
+	bool bHavePendingUnwrapHit = false;
+	bool bHavePendingLivePreviewHit = false;
+
+	int32 ShiftModifierID = 1;
+	int32 CtrlModifierID = 2;
+	bool bShiftToggle = false;
+	bool bCtrlToggle = false;
+	bool bCurrentStrokeIsSubtracting = false;
+
+	// TODO: It would be nice to have Ctrl+Shift brush inverting, so that a user could brush
+	//  an inverted selection. However that requires a bit of tedium, so we'll leave it as
+	//  a todo. We can't just update the selection as we go along in that case- we have to 
+	//  store the pre-stroke selection, additively add to a stroke selection, and then use
+	//  the two to update the actual selection mechanic. Plus we have to do that both for
+	//  unset selection and regular selection. Still it would be nice to add, especially if we add
+	//  that to the regular select tool.
+	// bool bCurrentStrokeIsInverting = false;
+	
+	TArray<int32> TempROIBuffer;
+
+	// Currently used just so that the Esc key exits out of the tool instead of clearing
+	//  the selection if we haven't yet done anything.
+	bool bHaveInteracted = false;
+
+	bool bHoveringUnwrap = false;
+	bool bHoveringLivePreview = false;
+	bool bDraggingUnwrap = false;
+	bool bDraggingLivePreview = false;
+	void UpdateViewportStateFromHoverOrDragEvent(bool bFromUnwrap, bool bIsEndEvent, bool bDragging);
+	
+	void IncreaseBrushRadiusAction();
+	void DecreaseBrushRadiusAction();
+};
\ No newline at end of file
diff --git a/Engine/Plugins/Editor/UVEditor/Source/UVEditorTools/Public/UVEditorToolBase.h b/Engine/Plugins/Editor/UVEditor/Source/UVEditorTools/Public/UVEditorToolBase.h
new file mode 100644
--- /dev/null
+++ b/Engine/Plugins/Editor/UVEditor/Source/UVEditorTools/Public/UVEditorToolBase.h
@@ -0,0 +1,56 @@
+// Copyright Epic Games, Inc. All Rights Reserved.
+
+#pragma once
+
+#include "InteractiveToolBuilder.h"
+#include "Templates/SubclassOf.h"
+#include "UObject/Interface.h"
+
+#include "UVEditorToolBase.generated.h"
+
+class UUVEditorToolMeshInput;
+
+/** UObject for IUVEditorGenericBuildableTool */
+UINTERFACE(MinimalAPI)
+class UUVEditorGenericBuildableTool : public UInterface
+{
+	GENERATED_BODY()
+};
+
+/**
+ * Interface that allows a component to receive various gizmo-specific callbacks while
+ * still inheriting from some class other than UGizmoBaseComponent.
+ */
+class IUVEditorGenericBuildableTool
+{
+	GENERATED_BODY()
+public:
+	virtual void SetTargets(const TArray<TObjectPtr<UUVEditorToolMeshInput>>& TargetsIn) = 0;
+};
+
+//~ TODO: We can use this builder for pretty much all our UV tools and remove the other builder classes
+/**
+ * Simple builder that just instantiates the given class and passes in the targets. Can be used
+ * for any UV tools that don't need special handling, as long as they implement IUVEditorGenericBuildableTool.
+ */
+UCLASS()
+class UVEDITORTOOLS_API UGenericUVEditorToolBuilder : public UInteractiveToolBuilder
+{
+	GENERATED_BODY()
+public:
+
+	/**
+	 * @param ToolClassIn Must implement IUVEditorGenericBuildableTool.
+	 */
+	virtual void Initialize(TArray<TObjectPtr<UUVEditorToolMeshInput>>& TargetsIn, TSubclassOf<UInteractiveTool> ToolClassIn);
+
+	virtual bool CanBuildTool(const FToolBuilderState& SceneState) const override;
+	virtual UInteractiveTool* BuildTool(const FToolBuilderState& SceneState) const override;
+
+private:
+	// This is a pointer so that it can be updated under the builder without
+	// having to set it in the mode after initializing targets.
+	const TArray<TObjectPtr<UUVEditorToolMeshInput>>* Targets = nullptr;
+
+	TSubclassOf<UInteractiveTool> ToolClass = nullptr;
+};
\ No newline at end of file
diff --git a/Engine/Plugins/Editor/UVEditor/Source/UVEditorTools/UVEditorTools.Build.cs b/Engine/Plugins/Editor/UVEditor/Source/UVEditorTools/UVEditorTools.Build.cs
--- a/Engine/Plugins/Editor/UVEditor/Source/UVEditorTools/UVEditorTools.Build.cs
+++ b/Engine/Plugins/Editor/UVEditor/Source/UVEditorTools/UVEditorTools.Build.cs
@@ -47,6 +47,7 @@ public class UVEditorTools : ModuleRules
 					
 				"DynamicMesh",
 				"GeometryCore",
+				"InputCore", // EKey
 				"ModelingComponents",
 				"ModelingOperators",
 				"TextureUtilitiesCommon"
