Modeling Tools: Adds min value of 0 to UV Editor display settings of line thickness/point size

diff --git a/Engine/Plugins/Editor/UVEditor/Source/UVEditor/Private/UVEditorUXPropertySets.h b/Engine/Plugins/Editor/UVEditor/Source/UVEditor/Private/UVEditorUXPropertySets.h
--- a/Engine/Plugins/Editor/UVEditor/Source/UVEditor/Private/UVEditorUXPropertySets.h
+++ b/Engine/Plugins/Editor/UVEditor/Source/UVEditor/Private/UVEditorUXPropertySets.h
@@ -17,11 +17,11 @@ public:
 	UPROPERTY(EditAnywhere, Category="Unwrapped Visualization", meta = (DisplayName = "Boundary Line Color"))
 	TArray<FColor> BoundaryLineColors;
 
-	UPROPERTY(EditAnywhere, Category="Unwrapped Visualization", meta = (DisplayName = "Boundary Line Thickness"))
+	UPROPERTY(EditAnywhere, Category="Unwrapped Visualization", meta = (UIMin = "0.0", UIMax = "10.0", ClampMin = "0.0", DisplayName = "Boundary Line Thickness"))
 	float BoundaryLineThickness;
 
 	// TODO: maybe opacity is the setting to change instead of/in addition to thickness? Could also have something to recognize line density and/or window size and adjust accordingly
-	UPROPERTY(EditAnywhere, Category="Unwrapped Visualization", meta = (DisplayName = "Wireframe Thickness"))
+	UPROPERTY(EditAnywhere, Category="Unwrapped Visualization", meta = (UIMin = "0.0", UIMax = "10.0", ClampMin = "0.0", DisplayName = "Wireframe Thickness"))
 	float WireframeThickness;
 };
 
@@ -36,9 +36,9 @@ public:
 	UPROPERTY(EditAnywhere, Category="Live Preview Visualization", meta = (DisplayName = "Selection Line Color"))
 	FColor SelectionColor;
 
-	UPROPERTY(EditAnywhere, Category="Live Preview Visualization", meta = (DisplayName = "Selection Line Thickness"))
+	UPROPERTY(EditAnywhere, Category="Live Preview Visualization", meta = (UIMin = "0.0", UIMax = "10.0", ClampMin = "0.0", DisplayName = "Selection Line Thickness"))
 	float SelectionLineThickness;
 
-	UPROPERTY(EditAnywhere, Category="Live Preview Visualization", meta = (DisplayName = "Selection Point Size"))
+	UPROPERTY(EditAnywhere, Category="Live Preview Visualization", meta = (UIMin = "0.0", UIMax = "10.0", ClampMin = "0.0", DisplayName = "Selection Point Size"))
 	float SelectionPointSize;
 };
\ No newline at end of file
