Modeling Tools: Adds UV Snapshot Tool

diff --git a/Engine/Plugins/Editor/UVEditor/Source/UVEditor/Private/UVEditorCommands.cpp b/Engine/Plugins/Editor/UVEditor/Source/UVEditor/Private/UVEditorCommands.cpp
--- a/Engine/Plugins/Editor/UVEditor/Source/UVEditor/Private/UVEditorCommands.cpp
+++ b/Engine/Plugins/Editor/UVEditor/Source/UVEditor/Private/UVEditorCommands.cpp
@@ -35,6 +35,7 @@ void FUVEditorCommands::RegisterCommands()
 	UI_COMMAND(BeginSeamTool, "Seam", "Edit UV seams", EUserInterfaceActionType::ToggleButton, FInputChord());
 	UI_COMMAND(BeginRecomputeUVsTool, "Unwrap", "Perform UV unwrapping", EUserInterfaceActionType::ToggleButton, FInputChord());
 	UI_COMMAND(BeginBrushSelectTool, "Brush", "Brush select triangles", EUserInterfaceActionType::ToggleButton, FInputChord());
+	UI_COMMAND(BeginUVSnapshotTool, "Snapshot", "Export a texture asset of a UV Layout", EUserInterfaceActionType::ToggleButton, FInputChord());
 
 	// These get linked to one-off tool actions.
 	UI_COMMAND(SewAction, "Sew", "Sew edges highlighted in red to edges highlighted in green", EUserInterfaceActionType::Button, FInputChord());
diff --git a/Engine/Plugins/Editor/UVEditor/Source/UVEditor/Private/UVEditorMode.cpp b/Engine/Plugins/Editor/UVEditor/Source/UVEditor/Private/UVEditorMode.cpp
--- a/Engine/Plugins/Editor/UVEditor/Source/UVEditor/Private/UVEditorMode.cpp
+++ b/Engine/Plugins/Editor/UVEditor/Source/UVEditor/Private/UVEditorMode.cpp
@@ -16,6 +16,7 @@
 #include "Selection/UVToolSelectionAPI.h"
 #include "Drawing/MeshElementsVisualizer.h"
 #include "Editor.h"
+#include "EditorModelingObjectsCreationAPI.h"
 #include "EditorModes.h"
 #include "EditorViewportClient.h"
 #include "EdModeInteractiveToolsContext.h" //ToolsContext, EditorInteractiveToolsContext
@@ -44,6 +45,7 @@
 #include "UVEditorLayerEditTool.h"
 #include "UVEditorSeamTool.h"
 #include "UVEditorRecomputeUVsTool.h"
+#include "UVEditorUVSnapshotTool.h"
 #include "UVSelectTool.h"
 #include "UVEditorTexelDensityTool.h"
 #include "UVEditorInitializationContext.h"
@@ -632,6 +634,9 @@ void UUVEditorMode::RegisterTools()
 	BrushToolBuilder->Initialize(ToolInputObjects, UUVEditorBrushSelectTool::StaticClass());
 	RegisterTool(CommandInfos.BeginBrushSelectTool, BrushToolIdentifier, BrushToolBuilder);
 	ToolsThatAllowActions.Add(BrushToolIdentifier);
+	UUVEditorUVSnapshotToolBuilder* UVEditorUVSnapshotToolBuilder = NewObject<UUVEditorUVSnapshotToolBuilder>();
+	UVEditorUVSnapshotToolBuilder->Targets = &ToolInputObjects;
+	RegisterTool(CommandInfos.BeginUVSnapshotTool, TEXT("BeginUVSnapshotTool"), UVEditorUVSnapshotToolBuilder);
 }
 
 void UUVEditorMode::RegisterActions()
@@ -1063,6 +1068,15 @@ void UUVEditorMode::InitializeAssetEditorContexts(UContextObjectStore& ContextSt
 	const TArray<TObjectPtr<UObject>>& AssetsIn, const TArray<FTransform>& TransformsIn,
 	FEditorViewportClient& LivePreviewViewportClient, FAssetEditorModeManager& LivePreviewModeManager,
 	UUVToolViewportButtonsAPI& ViewportButtonsAPI, UUVTool2DViewportAPI& UVTool2DViewportAPI)
+{
+	UInteractiveToolsContext* InteractiveToolsContext = Cast<UInteractiveToolsContext>(ContextStore.GetOuter());
+	InitializeAssetEditorContexts(ContextStore, AssetsIn, TransformsIn, LivePreviewViewportClient, LivePreviewModeManager,
+		ViewportButtonsAPI, UVTool2DViewportAPI, *InteractiveToolsContext);
+}
+void UUVEditorMode::InitializeAssetEditorContexts(UContextObjectStore& ContextStore,
+	const TArray<TObjectPtr<UObject>>& AssetsIn, const TArray<FTransform>& TransformsIn,
+	FEditorViewportClient& LivePreviewViewportClient, FAssetEditorModeManager& LivePreviewModeManager,
+	UUVToolViewportButtonsAPI& ViewportButtonsAPI, UUVTool2DViewportAPI& UVTool2DViewportAPI, UInteractiveToolsContext& ToolsContext)
 {
 	using namespace UVEditorModeLocals;
 
@@ -1074,6 +1088,13 @@ void UUVEditorMode::InitializeAssetEditorContexts(UContextObjectStore& ContextSt
 		ContextStore.AddContextObject(AssetInputsContext);
 	}
 
+	UEditorModelingObjectsCreationAPI* ModelingObjectsCreationAPI = ContextStore.FindContext<UEditorModelingObjectsCreationAPI>();
+	if (!ModelingObjectsCreationAPI)
+	{
+		ModelingObjectsCreationAPI = NewObject<UEditorModelingObjectsCreationAPI>(&ToolsContext);
+		ContextStore.AddContextObject(ModelingObjectsCreationAPI);
+	}
+	
 	UToolsContextCursorAPI* LivePreviewToolsContextCursorAPI = LivePreviewModeManager.GetInteractiveToolsContext()->ContextObjectStore->FindContext<UToolsContextCursorAPI>();
 	UUVToolLivePreviewAPI* LivePreviewAPI = ContextStore.FindContext<UUVToolLivePreviewAPI>();
 	if (!LivePreviewAPI)
diff --git a/Engine/Plugins/Editor/UVEditor/Source/UVEditor/Private/UVEditorModeToolkit.cpp b/Engine/Plugins/Editor/UVEditor/Source/UVEditor/Private/UVEditorModeToolkit.cpp
--- a/Engine/Plugins/Editor/UVEditor/Source/UVEditor/Private/UVEditorModeToolkit.cpp
+++ b/Engine/Plugins/Editor/UVEditor/Source/UVEditor/Private/UVEditorModeToolkit.cpp
@@ -493,6 +493,7 @@ void FUVEditorModeToolkit::BuildToolPalette(FName PaletteIndex, class FToolBarBu
 		ToolbarBuilder.AddToolBarButton(Commands.BeginParameterizeMeshTool);
 		ToolbarBuilder.AddToolBarButton(Commands.BeginRecomputeUVsTool);
 		ToolbarBuilder.AddToolBarButton(Commands.BeginBrushSelectTool);
+		ToolbarBuilder.AddToolBarButton(Commands.BeginUVSnapshotTool);
 	}
 }
 
diff --git a/Engine/Plugins/Editor/UVEditor/Source/UVEditor/Private/UVEditorToolkit.cpp b/Engine/Plugins/Editor/UVEditor/Source/UVEditor/Private/UVEditorToolkit.cpp
--- a/Engine/Plugins/Editor/UVEditor/Source/UVEditor/Private/UVEditorToolkit.cpp
+++ b/Engine/Plugins/Editor/UVEditor/Source/UVEditor/Private/UVEditorToolkit.cpp
@@ -453,7 +453,7 @@ void FUVEditorToolkit::PostInitAssetEditor()
 	// the context store on Enter() to function properly.
 	UUVEditorMode::InitializeAssetEditorContexts(*EditorModeManager->GetInteractiveToolsContext()->ContextObjectStore, 
 		ObjectsToEdit, ObjectTransforms, *LivePreviewViewportClient, *LivePreviewEditorModeManager, 
-		*ViewportButtonsAPI, *UVTool2DViewportAPI);
+		*ViewportButtonsAPI, *UVTool2DViewportAPI, *EditorModeManager->GetInteractiveToolsContext());
 
 	// Currently, aside from setting up all the UI elements, the toolkit also kicks off the UV
 	// editor mode, which is the mode that the editor always works in (things are packaged into
diff --git a/Engine/Plugins/Editor/UVEditor/Source/UVEditor/Public/UVEditorCommands.h b/Engine/Plugins/Editor/UVEditor/Source/UVEditor/Public/UVEditorCommands.h
--- a/Engine/Plugins/Editor/UVEditor/Source/UVEditor/Public/UVEditorCommands.h
+++ b/Engine/Plugins/Editor/UVEditor/Source/UVEditor/Public/UVEditorCommands.h
@@ -29,6 +29,7 @@ public:
 	TSharedPtr<FUICommandInfo> BeginSeamTool;
 	TSharedPtr<FUICommandInfo> BeginRecomputeUVsTool;
 	TSharedPtr<FUICommandInfo> BeginBrushSelectTool;
+	TSharedPtr<FUICommandInfo> BeginUVSnapshotTool;
 
 	TSharedPtr<FUICommandInfo> SewAction;
 	TSharedPtr<FUICommandInfo> SplitAction;
diff --git a/Engine/Plugins/Editor/UVEditor/Source/UVEditor/Public/UVEditorMode.h b/Engine/Plugins/Editor/UVEditor/Source/UVEditor/Public/UVEditorMode.h
--- a/Engine/Plugins/Editor/UVEditor/Source/UVEditor/Public/UVEditorMode.h
+++ b/Engine/Plugins/Editor/UVEditor/Source/UVEditor/Public/UVEditorMode.h
@@ -8,6 +8,7 @@
 #include "ToolTargets/ToolTarget.h" // FToolTargetTypeRequirements
 #include "GeometryBase.h"
 #include "InteractiveTool.h"
+#include "InteractiveToolsContext.h"
 
 #include "UVEditorMode.generated.h"
 
@@ -175,10 +176,15 @@ public:
 	/**
 	 * Called by an asset editor so that a created instance of the mode has all the data it needs on Enter() to initialize itself.
 	 */
+	UE_DEPRECATED(5.5, "Initialization of contexts now requires UInteractiveToolsContext to be passed in directly. Please update code accordingly")
 	static void InitializeAssetEditorContexts(UContextObjectStore& ContextStore,
 		const TArray<TObjectPtr<UObject>>& AssetsIn, const TArray<FTransform>& TransformsIn,
 		FEditorViewportClient& LivePreviewViewportClient, FAssetEditorModeManager& LivePreviewModeManager,
 		UUVToolViewportButtonsAPI& ViewportButtonsAPI, UUVTool2DViewportAPI& UVTool2DViewportAPI);
+	static void InitializeAssetEditorContexts(UContextObjectStore& ContextStore,
+		const TArray<TObjectPtr<UObject>>& AssetsIn, const TArray<FTransform>& TransformsIn,
+		FEditorViewportClient& LivePreviewViewportClient, FAssetEditorModeManager& LivePreviewModeManager,
+		UUVToolViewportButtonsAPI& ViewportButtonsAPI, UUVTool2DViewportAPI& UVTool2DViewportAPI, UInteractiveToolsContext& ToolsContext);
 
 	// Public for use by undo/redo. Otherwise should use RequestUVChannelChange
 	void ChangeInputObjectLayer(int32 AssetID, int32 NewLayerIndex);
diff --git a/Engine/Plugins/Editor/UVEditor/Source/UVEditorTools/Private/UVEditorUVSnapshotTool.cpp b/Engine/Plugins/Editor/UVEditor/Source/UVEditorTools/Private/UVEditorUVSnapshotTool.cpp
new file mode 100644
--- /dev/null
+++ b/Engine/Plugins/Editor/UVEditor/Source/UVEditorTools/Private/UVEditorUVSnapshotTool.cpp
@@ -0,0 +1,356 @@
+// Copyright Epic Games, Inc. All Rights Reserved.
+#include "UVEditorUVSnapshotTool.h"
+
+#include "AssetToolsModule.h"
+#include "ContentBrowserModule.h"
+#include "IContentBrowserSingleton.h"
+#include "ModelingToolTargetUtil.h"
+#include "AssetUtils/Texture2DBuilder.h"
+#include "Sampling/MeshUVShellMapEvaluator.h"
+#include "ModelingObjectsCreationAPI.h"
+#include "Materials/Material.h"
+#include "Materials/MaterialInstanceDynamic.h"
+#include "TargetInterfaces/StaticMeshBackedTarget.h"
+#include "UVEditorUXSettings.h"
+#include "Drawing/MeshElementsVisualizer.h"
+
+#include UE_INLINE_GENERATED_CPP_BY_NAME(UVEditorUVSnapshotTool)
+
+using namespace UE::Geometry;
+
+#define LOCTEXT_NAMESPACE "UUVEditorUVSnapshotTool"
+
+/*
+ * ToolBuilder
+ */
+bool UUVEditorUVSnapshotToolBuilder::CanBuildTool(const FToolBuilderState& SceneState) const
+{
+	return Targets && Targets->Num() == 1;
+}
+UInteractiveTool* UUVEditorUVSnapshotToolBuilder::BuildTool(const FToolBuilderState& SceneState) const
+{
+	UUVEditorUVSnapshotTool* NewTool = NewObject<UUVEditorUVSnapshotTool>(SceneState.ToolManager);
+	NewTool->SetTarget((*Targets)[0]);
+	return NewTool;
+}
+
+/*
+ * Operator
+ */
+class FMeshUVMapBakerOp : public TGenericDataOperator<FMeshMapBaker>
+{
+public:
+	using ImagePtr = TSharedPtr<UE::Geometry::TImageBuilder<FVector4f>, ESPMode::ThreadSafe>;
+	// General bake settings
+	TSharedPtr<UE::Geometry::FDynamicMesh3, ESPMode::ThreadSafe> BaseMesh;
+	TUniquePtr<UE::Geometry::FMeshMapBaker> Baker;
+	UUVEditorUVSnapshotTool::FUVSnapshotBakeSettings BakeSettings;
+	TSharedPtr<UE::Geometry::FDynamicMeshAABBTree3, ESPMode::ThreadSafe> DetailSpatial;
+
+	virtual void CalculateResult(FProgressCancel* Progress) override
+	{
+		Baker = MakeUnique<FMeshMapBaker>();
+		Baker->CancelF = [Progress]() {
+			return Progress && Progress->Cancelled();
+		};
+		Baker->SetTargetMesh(BaseMesh.Get());
+		Baker->SetTargetMeshUVLayer(BakeSettings.UVLayer);
+		Baker->SetDimensions(BakeSettings.Dimensions);
+		Baker->SetSamplesPerPixel(BakeSettings.SamplesPerPixel);
+		
+		FMeshBakerDynamicMeshSampler DetailSampler(BaseMesh.Get(), DetailSpatial.Get(), nullptr);
+		Baker->SetDetailSampler(&DetailSampler);
+		const TSharedPtr<FMeshUVShellMapEvaluator> UVShellEval = MakeShared<FMeshUVShellMapEvaluator>();
+		UVShellEval->TexelSize = BakeSettings.Dimensions.GetTexelSize();
+		UVShellEval->UVLayer = BakeSettings.UVLayer;
+		UVShellEval->WireframeThickness = BakeSettings.WireframeThickness;
+		UVShellEval->WireframeColor = BakeSettings.WireframeColor;
+		UVShellEval->ShellColor = BakeSettings.ShellColor;
+		UVShellEval->BackgroundColor = BakeSettings.BackgroundColor;
+		Baker->AddEvaluator(UVShellEval);
+		Baker->Bake();
+		SetResult(MoveTemp(Baker));
+	}
+};
+
+/*
+ * Tool
+ */
+void UUVEditorUVSnapshotTool::Setup()
+{
+	UInteractiveTool::Setup();
+
+	// initialize properties
+	UVShellSettings = NewObject<UUVEditorBakeUVShellProperties>(this);
+	UVShellSettings->RestoreProperties(this);
+	AddToolPropertySource(UVShellSettings);
+
+	// retrieve whatever UV Layer is currently being displayed in UV Editor
+	UVShellSettings->UVLayer = Target->UVLayerIndex;
+	
+	UVShellSettings->WatchProperty(UVShellSettings->SamplesPerPixel, [this](EBakeTextureSamplesPerPixel) { OpState |= EUVSnapshotBakeOpState::Evaluate; });
+	UVShellSettings->WatchProperty(UVShellSettings->Resolution, [this](EBakeTextureResolution) { OpState |= EUVSnapshotBakeOpState::Evaluate; });
+	UVShellSettings->WatchProperty(UVShellSettings->UVLayer, [this](int) { OpState |= EUVSnapshotBakeOpState::Evaluate; });
+	UVShellSettings->WatchProperty(UVShellSettings->WireframeThickness, [this](float) { OpState |= EUVSnapshotBakeOpState::Evaluate; });
+	UVShellSettings->WatchProperty(UVShellSettings->WireframeColor, [this](FLinearColor) { OpState |= EUVSnapshotBakeOpState::Evaluate; });
+	UVShellSettings->WatchProperty(UVShellSettings->ShellColor, [this](FLinearColor) { OpState |= EUVSnapshotBakeOpState::Evaluate; });
+	UVShellSettings->WatchProperty(UVShellSettings->BackgroundColor, [this](FLinearColor) { OpState |= EUVSnapshotBakeOpState::Evaluate; });
+	SetToolPropertySourceEnabled(UVShellSettings, true);
+	
+	ResultSettings = NewObject<UUVEditorBakeUVShellResultProperties>(this);
+	ResultSettings->RestoreProperties(this);
+	AddToolPropertySource(ResultSettings);
+	SetToolPropertySourceEnabled(ResultSettings, true);
+	ResultSettings->Result = nullptr;
+
+	// set up the detail mesh & spatial
+	FDynamicMesh3 DetailMeshGet;
+	Target->AppliedPreview->GetCurrentResultCopy(DetailMeshGet, false);
+	DetailMesh = MakeShared<FDynamicMesh3, ESPMode::ThreadSafe>(DetailMeshGet);
+	
+	DetailSpatial = MakeShared<FDynamicMeshAABBTree3, ESPMode::ThreadSafe>();
+	DetailSpatial->SetMesh(DetailMesh.Get(), true);
+
+	// mark for evaluation
+	OpState |= EUVSnapshotBakeOpState::Evaluate;
+
+	CachedBakeSettings = FUVSnapshotBakeSettings();
+
+	// set up UPreviewGeometry for visualization in Unwrap viewport
+	PreviewGeoBackgroundQuad = NewObject<UPreviewGeometry>(this);
+	PreviewGeoBackgroundQuad->CreateInWorld(Target->UnwrapPreview->GetWorld(), FTransform::Identity);
+	PreviewGeoBackgroundQuad->AddTriangleSet("UVShellMap");
+	PreviewGeoBackgroundQuad->SetAllVisible(true);
+
+	// set up op factory
+	SetToolDisplayName(LOCTEXT("ToolNameLocal", "UV Snapshot"));
+	GetToolManager()->DisplayMessage(
+		LOCTEXT("OnStartUVSnapshotTool", "Export a texture asset of a UV Layout."),
+		EToolMessageLevel::UserNotification);
+	
+}
+void UUVEditorUVSnapshotTool::Shutdown(EToolShutdownType ShutdownType)
+{
+	if (ShutdownType == EToolShutdownType::Accept)
+	{
+		const IStaticMeshBackedTarget* StaticMeshTarget = Cast<IStaticMeshBackedTarget>(Target->SourceTarget);
+		UObject* SourceAsset = StaticMeshTarget ? StaticMeshTarget->GetStaticMesh() : nullptr;
+		const UPrimitiveComponent* SourceComponent = UE::ToolTarget::GetTargetComponent(Target->SourceTarget);
+		CreateTextureAsset(ResultSettings->Result, SourceComponent->GetWorld(), SourceAsset);
+	}
+	
+	UVShellSettings->SaveProperties(this);
+	if (UVShellSettings)
+	{
+		UVShellSettings->RestoreProperties(this, TEXT("UVEditorUVSnapshotTool"));
+		UVShellSettings = nullptr;
+	}
+	if (Compute)
+	{
+		Compute->Shutdown();
+	}
+
+	Target->AppliedPreview->ClearOpFactory();
+
+	if (PreviewGeoBackgroundQuad)
+	{
+		PreviewGeoBackgroundQuad->Disconnect();
+		PreviewGeoBackgroundQuad = nullptr;
+	}
+	// re-enable wireframe display and unwrap preview
+	Target->WireframeDisplay->Settings->bVisible = true;
+	Target->UnwrapPreview->SetVisibility(true);
+}
+void UUVEditorUVSnapshotTool::OnTick(float DeltaTime)
+{
+	if (Compute)
+	{
+		Compute->Tick(DeltaTime);
+	}
+}
+void UUVEditorUVSnapshotTool::Render(IToolsContextRenderAPI* RenderAPI)
+{
+	UpdateResult();
+}
+TUniquePtr<UE::Geometry::TGenericDataOperator<UE::Geometry::FMeshMapBaker>> UUVEditorUVSnapshotTool::MakeNewOperator()
+{
+	TUniquePtr<FMeshUVMapBakerOp> Op = MakeUnique<FMeshUVMapBakerOp>();
+	Op->DetailSpatial = DetailSpatial;
+	Op->BaseMesh = DetailMesh;
+	Op->BakeSettings = CachedBakeSettings;
+	
+	return Op;
+}
+void UUVEditorUVSnapshotTool::UpdateResult()
+{
+	if (OpState == EUVSnapshotBakeOpState::Clean)
+	{
+		return;
+	}
+	// clear warning (ugh)
+	GetToolManager()->DisplayMessage(FText(), EToolMessageLevel::UserWarning);
+
+	const int32 ImageSize = (int32)UVShellSettings->Resolution;
+	const FImageDimensions Dimensions(ImageSize, ImageSize);
+	CachedBakeSettings.Dimensions = Dimensions;
+	
+	// Clear our invalid bitflag to check again for valid inputs.
+	OpState &= ~EUVSnapshotBakeOpState::Invalid;
+	OpState |= UpdateResult_UVShellMap();
+	
+	// Early exit if op input parameters are invalid.
+	if ((bool)(OpState & EUVSnapshotBakeOpState::Invalid))
+	{
+		InvalidateResults();
+		return;
+	}
+	// This should be the only point of compute invalidation to
+	// minimize synchronization issues.
+	InvalidateCompute();
+}
+void UUVEditorUVSnapshotTool::OnMapUpdated(const TUniquePtr<UE::Geometry::FMeshMapBaker>& NewResult)
+{
+	const FImageDimensions BakeDimensions = NewResult->GetDimensions();
+	constexpr bool bConvertToSRGB = true;
+	FTexture2DBuilder TextureBuilder;
+	TextureBuilder.Initialize(FTexture2DBuilder::ETextureType::Color, BakeDimensions);
+	TextureBuilder.Copy(*NewResult->GetBakeResults(0)[0], bConvertToSRGB);
+	TextureBuilder.Commit(false);
+	// Copy image to source data after commit. This will avoid incurring
+	// the cost of hitting the DDC for texture compile while iterating on
+	// bake settings. Since this dirties the texture, the next time the texture
+	// is used after accepting the final texture, the DDC will trigger and
+	// properly recompile the platform data.
+	constexpr bool bConvertSourceToSRGB = true;
+	// default should be ChannelBits8
+	constexpr ETextureSourceFormat SourceDataFormat = TSF_BGRA8;
+	TextureBuilder.CopyImageToSourceData(*NewResult->GetBakeResults(0)[0], SourceDataFormat, bConvertSourceToSRGB);
+	// The CachedUVMap can be thrown out of sync if updated during a background
+	// compute. Validate the computed type against our cached map.
+	CachedUVMap = TextureBuilder.GetTexture2D();
+	UpdateVisualization();
+	GetToolManager()->PostInvalidation();
+	
+}
+void UUVEditorUVSnapshotTool::InvalidateResults() const
+{
+	ResultSettings->Result = nullptr;
+}
+void UUVEditorUVSnapshotTool::InvalidateCompute()
+{
+	if (!Compute)
+	{
+		// Initialize background compute
+		Compute = MakeUnique<TGenericDataBackgroundCompute<FMeshMapBaker>>();
+		Compute->Setup(this);
+		Compute->OnResultUpdated.AddLambda([this](const TUniquePtr<FMeshMapBaker>& NewResult) { OnMapUpdated(NewResult); });
+	}
+	Compute->InvalidateResult();
+	OpState = EUVSnapshotBakeOpState::Clean;
+}
+void UUVEditorUVSnapshotTool::UpdateVisualization() 
+{
+	ResultSettings->Result = CachedUVMap;
+	UpdatePreviewMaterialBasedOnBackground();
+}
+EUVSnapshotBakeOpState UUVEditorUVSnapshotTool::UpdateResult_UVShellMap()
+{
+	EUVSnapshotBakeOpState ResultState = EUVSnapshotBakeOpState::Clean;
+	FUVSnapshotBakeSettings UVShellMapSettings;
+	UVShellMapSettings.UVLayer = UVShellSettings->UVLayer;
+	UVShellMapSettings.WireframeThickness = UVShellSettings->WireframeThickness;
+	UVShellMapSettings.WireframeColor = UVShellSettings->WireframeColor;
+	UVShellMapSettings.ShellColor = UVShellSettings->ShellColor;
+	UVShellMapSettings.BackgroundColor = UVShellSettings->BackgroundColor;
+
+	const int32 ImageSize = (int32)UVShellSettings->Resolution;
+	const FImageDimensions Dimensions(ImageSize, ImageSize);
+	UVShellMapSettings.Dimensions = Dimensions;
+	UVShellMapSettings.SamplesPerPixel = (int32)UVShellSettings->SamplesPerPixel;
+
+	// if our settings have changed
+	if (CachedBakeSettings != UVShellMapSettings)
+	{
+		CachedBakeSettings = UVShellMapSettings;
+		ResultState |= EUVSnapshotBakeOpState::Evaluate;
+	}
+	return ResultState;
+}
+void UUVEditorUVSnapshotTool::UpdatePreviewMaterialBasedOnBackground()
+{
+	// using this material so that we can set a texture
+	UMaterialInstanceDynamic* Mat = UMaterialInstanceDynamic::Create(LoadObject<UMaterial>(nullptr, TEXT("/UVEditor/Materials/UVEditorBackground")), this);
+	Mat->SetTextureParameterValue(TEXT("BackgroundBaseMap_Color"), ResultSettings->Result);
+	Mat->SetScalarParameterValue(TEXT("BackgroundPixelDepthOffset"), FUVEditorUXSettings::BackgroundQuadDepthOffset - 0.5);	
+
+	// temporarily disable wireframe overlay and unwrap preview; re-enabled on shutdown
+	Target->WireframeDisplay->Settings->bVisible = false;
+	Target->UnwrapPreview->SetVisibility(false);
+
+	// connect to existing Preview Geometry
+	UTriangleSetComponent* TriangleSet = PreviewGeoBackgroundQuad->FindTriangleSet("UVShellMap");
+	TriangleSet->Clear();
+
+	const FVector Normal(0, 0, 1);
+	const FColor BackgroundColor = FColor::Red;
+
+	// set up rendering of 2 triangles to make one 2d quad
+	FIndex2i UDimBlockToRender = FIndex2i(0,0);
+
+	auto MakeQuadVert = [&UDimBlockToRender, &Normal, &BackgroundColor](int32 CornerX, int32 CornerY)
+	{
+		FVector2f ExternalUV(static_cast<float>(UDimBlockToRender.A + CornerX), static_cast<float>(UDimBlockToRender.B + CornerY));
+
+		return FRenderableTriangleVertex(FUVEditorUXSettings::ExternalUVToUnwrapWorldPosition(ExternalUV),
+			(FVector2D)FUVEditorUXSettings::ExternalUVToInternalUV(ExternalUV),
+			Normal, BackgroundColor);
+	};
+	
+	FRenderableTriangleVertex V00 = MakeQuadVert(0,0);
+	FRenderableTriangleVertex V10 = MakeQuadVert(1, 0);
+	FRenderableTriangleVertex V11 = MakeQuadVert(1, 1);
+	FRenderableTriangleVertex V01 = MakeQuadVert(0,1);
+
+	// add to TriangleSet in PreviewGeometry
+	TriangleSet->AddTriangle(FRenderableTriangle(Mat, V00, V10, V11));
+	TriangleSet->AddTriangle(FRenderableTriangle(Mat, V00, V11, V01));
+	
+}
+void UUVEditorUVSnapshotTool::CreateTextureAsset(const TObjectPtr<UTexture2D>& Texture, UWorld* SourceWorld, UObject* SourceAsset) const
+{
+	bool bCreatedAssetOK = true;
+	const FString ObjName = UE::ToolTarget::GetTargetActor(Target->SourceTarget)->GetActorNameOrLabel();
+	FString NewAssetName = FString::Printf(TEXT("%s_UVShell_UV%d"), *ObjName, UVShellSettings->UVLayer); // will be something like "Cylinder_UVShell_UV0"
+
+	// open dialog so user can choose where to save out the new asset
+	IContentBrowserSingleton& ContentBrowser = FModuleManager::LoadModuleChecked<FContentBrowserModule>("ContentBrowser").Get();
+
+	FSaveAssetDialogConfig Config;
+	Config.DefaultAssetName = NewAssetName;
+	Config.DialogTitleOverride = LOCTEXT("GenerateStaticMeshActorPathDialogWarning", "Choose Folder Path and Name for New Asset. Cancel to Discard New Asset.");
+	// if we have previously saved a UVSnapshot, use that path as default; otherwise default path
+	Config.DefaultPath = UVShellSettings->SavedPath.IsEmpty() ?  FString() : UVShellSettings->SavedPath;
+	
+	const FString SelectedPath = ContentBrowser.CreateModalSaveAssetDialog(Config);
+
+	// ensures that if save dialog is closed without saving, nothing happens
+	if (SelectedPath.IsEmpty() == false)
+	{
+		// save path so that if UV Snapshot is performed again, when save dialog opens we are already in previous location
+		UVShellSettings->SavedPath = FPaths::GetPath(SelectedPath);
+		NewAssetName = FPaths::GetBaseFilename(SelectedPath, true);
+
+		FString PackageNameOut, AssetNameOut;
+		const FAssetToolsModule& AssetToolsModule = FModuleManager::LoadModuleChecked<FAssetToolsModule>("AssetTools");
+		AssetToolsModule.Get().CreateUniqueAssetName(
+			FPaths::Combine(UVShellSettings->SavedPath, NewAssetName), TEXT(""),
+			PackageNameOut, AssetNameOut);
+
+		// create the asset
+		bCreatedAssetOK = bCreatedAssetOK && UE::Modeling::CreateTextureObject( GetToolManager(),
+			FCreateTextureObjectParams{ 0, SourceWorld, SourceAsset, NewAssetName, Texture}).IsOK();
+	}
+		
+	ensure(bCreatedAssetOK);
+}
+#undef LOCTEXT_NAMESPACE
\ No newline at end of file
diff --git a/Engine/Plugins/Editor/UVEditor/Source/UVEditorTools/Public/UVEditorUVSnapshotTool.h b/Engine/Plugins/Editor/UVEditor/Source/UVEditorTools/Public/UVEditorUVSnapshotTool.h
new file mode 100644
--- /dev/null
+++ b/Engine/Plugins/Editor/UVEditor/Source/UVEditorTools/Public/UVEditorUVSnapshotTool.h
@@ -0,0 +1,240 @@
+// // Copyright Epic Games, Inc. All Rights Reserved.
+#pragma once
+#include "Baking/BakingTypes.h"
+#include "CoreMinimal.h"
+#include "Drawing/PreviewGeometryActor.h"
+#include "InteractiveToolBuilder.h"
+#include "InteractiveToolQueryInterfaces.h"
+#include "MeshOpPreviewHelpers.h"
+#include "Sampling/MeshMapBaker.h"
+#include "Selection/UVToolSelectionAPI.h"
+
+#include "UVEditorUVSnapshotTool.generated.h"
+
+
+// Forward Declarations
+class UUVEditorBakeUVShellResultProperties;
+class UUVEditorBakeUVShellProperties;
+enum class EBakeMapType;
+
+/**
+ * Bake compute state
+ */
+enum class EUVSnapshotBakeOpState
+{
+	Clean              = 0,      // No-op - evaluation already launched/complete.
+	Evaluate           = 1 << 0, // Inputs are modified and valid, re-evaluate.
+	Invalid            = 1 << 1  // Inputs are modified and invalid - retry eval until valid.
+};
+ENUM_CLASS_FLAGS(EUVSnapshotBakeOpState);
+
+/**
+ * ToolBuilder
+ */
+UCLASS()
+class UVEDITORTOOLS_API UUVEditorUVSnapshotToolBuilder : public UInteractiveToolBuilder
+{
+	GENERATED_BODY()
+public:
+	virtual bool CanBuildTool(const FToolBuilderState& SceneState) const override;
+	virtual UInteractiveTool* BuildTool(const FToolBuilderState& SceneState) const override;
+	// This is a pointer so that it can be updated under the builder without
+	// having to set it in the mode after initializing targets.
+	// Supports only one target
+	const TArray<TObjectPtr<UUVEditorToolMeshInput>>* Targets = nullptr;
+};
+
+/**
+ * UV Snapshot Tool
+ *
+ * Exports a texture asset of a UV Layout
+ */
+UCLASS()
+class UVEDITORTOOLS_API UUVEditorUVSnapshotTool : public UInteractiveTool, public IUVToolSupportsSelection, public UE::Geometry::IGenericDataOperatorFactory<UE::Geometry::FMeshMapBaker>
+{
+	GENERATED_BODY()
+public:
+	// Begin UInteractiveTool interface
+	virtual void Setup() override;
+	virtual void Shutdown(EToolShutdownType ShutdownType) override;
+	
+	virtual void OnTick(float DeltaTime) override;
+	virtual void Render(IToolsContextRenderAPI* RenderAPI) override;
+	virtual bool HasCancel() const override { return true; }
+	virtual bool HasAccept() const override { return true; }
+	virtual bool CanAccept() const override { return true; }
+	// End UInteractiveTool interface
+
+	// Begin IGenericDataOperatorFactory interface
+	virtual TUniquePtr<UE::Geometry::TGenericDataOperator<UE::Geometry::FMeshMapBaker>> MakeNewOperator() override;
+	// End IGenericDataOperatorFactory interface
+	
+	/**
+	 * The tool will operate on the mesh given here. Supports only one mesh.
+	 */
+	virtual void SetTarget(const TObjectPtr<UUVEditorToolMeshInput>& TargetIn)
+	{
+		Target = TargetIn;
+	}
+protected:
+	//
+	// Mesh input to UV Editor
+	//
+	UPROPERTY()
+	TObjectPtr<UUVEditorToolMeshInput> Target;
+
+	//
+	// Property sets for bake and result
+	//
+	UPROPERTY()
+	TObjectPtr<UUVEditorBakeUVShellProperties> UVShellSettings;
+	UPROPERTY()
+	TObjectPtr<UUVEditorBakeUVShellResultProperties> ResultSettings;
+
+	//
+	// Preview Geometry for display in Unwrapped viewport
+	//
+	UPROPERTY()
+	TObjectPtr<UPreviewGeometry> PreviewGeoBackgroundQuad;
+
+	//
+	// Background compute
+	//
+	TUniquePtr<TGenericDataBackgroundCompute<UE::Geometry::FMeshMapBaker>> Compute = nullptr;
+
+public:
+	// Bake settings
+	struct FUVSnapshotBakeSettings
+	{
+		UE::Geometry::FImageDimensions Dimensions;
+		int UVLayer = 0;
+		float WireframeThickness = 1.0f;
+		FLinearColor WireframeColor = FLinearColor::Blue;
+		FLinearColor ShellColor = FLinearColor::Gray;
+		FLinearColor BackgroundColor = FLinearColor(0.0f, 0.0f, 0.0f, 0.0f);
+		int SamplesPerPixel = 4;
+		bool operator==(const FUVSnapshotBakeSettings& Other) const
+		{
+			return Dimensions == Other.Dimensions &&
+				UVLayer == Other.UVLayer &&
+				WireframeThickness == Other.WireframeThickness &&
+				WireframeColor == Other.WireframeColor &&
+				ShellColor == Other.ShellColor &&
+				BackgroundColor == Other.BackgroundColor &&
+				SamplesPerPixel == Other.SamplesPerPixel;
+		}
+	};
+protected:
+	FUVSnapshotBakeSettings CachedBakeSettings;
+	
+private:
+	EUVSnapshotBakeOpState OpState = EUVSnapshotBakeOpState::Evaluate;
+
+	// used in the bake
+	TSharedPtr<UE::Geometry::FDynamicMesh3, ESPMode::ThreadSafe> DetailMesh;
+	TSharedPtr<UE::Geometry::FDynamicMeshAABBTree3, ESPMode::ThreadSafe> DetailSpatial;
+
+	/**
+	 * Process OpState flags and update background compute
+	 * Invoked during render
+	 */
+	void UpdateResult();
+
+	/**
+	 * Retrieves the result of the FMeshMapBaker and generated UTexture2D into the CachedUVMap
+	 * 
+	 * @param NewResult the resulting FMeshMapBaker from the background compute
+	 */
+	void OnMapUpdated(const TUniquePtr<UE::Geometry::FMeshMapBaker>& NewResult);
+
+	/**
+	 * Sets the result to a nullptr
+	 */
+	void InvalidateResults() const;
+
+	/**
+	 * Invalidates the background compute operator.
+	 */
+	void InvalidateCompute();
+
+	/**
+	 * Updates the preview material on the preview mesh with the
+	 * computed results. Invoked by OnMapsUpdated.
+	 * Also sets the result to what is currently in CachedUVMap
+	 */
+	void UpdateVisualization();
+
+	/**
+	 * Updates the FUVSnapshotBakeSettings and CachedBakeSettings
+	 *
+	 * @return current result state (Clean or Evaluate if Bake settings have changed)
+	 */
+	EUVSnapshotBakeOpState UpdateResult_UVShellMap();
+
+	/**
+	 * Internal cache of bake uv texture result
+	 */
+	UPROPERTY()
+	TObjectPtr<UTexture2D> CachedUVMap = nullptr;
+
+	/**
+	 * Uses Preview Geometry to draw a preview of the bake in the unwrap viewport
+	 */
+	void UpdatePreviewMaterialBasedOnBackground();
+
+	/**
+	 * Create texture asset from our result Texture2D
+	 * @param Texture the result texture to create
+	 * @param SourceWorld the source world to define where the texture asset will be stored.
+	 * @param SourceAsset if not null, result texture will be stored adjacent to this asset.
+	 */
+	void CreateTextureAsset(const TObjectPtr<UTexture2D>& Texture, UWorld* SourceWorld, UObject* SourceAsset) const;
+};
+
+UCLASS()
+class UVEDITORTOOLS_API UUVEditorBakeUVShellResultProperties : public UInteractiveToolPropertySet
+{
+	GENERATED_BODY()
+public:
+	/** Bake */
+	UPROPERTY(VisibleAnywhere, Category = Results, meta = (DisplayName = "Result Texture", TransientToolProperty))
+	TObjectPtr<UTexture2D> Result;
+};
+
+UCLASS()
+class UVEDITORTOOLS_API UUVEditorBakeUVShellProperties : public UInteractiveToolPropertySet
+{
+	GENERATED_BODY()
+public:
+	/** The source mesh UV Layer to sample. */
+	UPROPERTY(EditAnywhere, Category = "UV Snapshot Output", meta = (DisplayName = "UV Layer"))
+	int UVLayer = 0;
+
+	/** The thickness of the wireframe in pixels. */
+	UPROPERTY(EditAnywhere, Category = "UV Snapshot Output", meta = (UIMin = "0.0", UIMax = "10.0", ClampMin = "0.0"))
+	float WireframeThickness = 1.0f;
+
+	/** The color of wireframe pixels. */
+	UPROPERTY(EditAnywhere, Category = "UV Snapshot Output")
+	FLinearColor WireframeColor = FLinearColor::Blue;
+
+	/** The color of the UV shell interior pixels. */
+	UPROPERTY(EditAnywhere, Category = "UV Snapshot Output")
+	FLinearColor ShellColor = FLinearColor::Gray;
+
+	/** The color of pixels external to UV shells. */
+	UPROPERTY(EditAnywhere, Category = "UV Snapshot Output")
+	FLinearColor BackgroundColor = FLinearColor(0.0f, 0.0f, 0.0f, 0.0f);
+
+	/** The pixel resolution of the generated textures */
+	UPROPERTY(EditAnywhere, Category = Textures)
+	EBakeTextureResolution Resolution = EBakeTextureResolution::Resolution256;
+
+	/** Number of samples per pixel */
+	UPROPERTY(EditAnywhere, Category = Textures)
+	EBakeTextureSamplesPerPixel SamplesPerPixel = EBakeTextureSamplesPerPixel::Sample4;
+
+	/** Saved path where last UVSnapshot was saved to. Empty if this is first save out */
+	UPROPERTY()
+	FString SavedPath = FString();
+};
\ No newline at end of file
