UVEditor: Make the seam tool (and the modeling tools seam tool) follow paths better in certain "against the grain" situations.

diff --git a/Engine/Plugins/Editor/UVEditor/Source/UVEditorTools/Private/UVEditorSeamTool.cpp b/Engine/Plugins/Editor/UVEditor/Source/UVEditorTools/Private/UVEditorSeamTool.cpp
--- a/Engine/Plugins/Editor/UVEditor/Source/UVEditorTools/Private/UVEditorSeamTool.cpp
+++ b/Engine/Plugins/Editor/UVEditor/Source/UVEditorTools/Private/UVEditorSeamTool.cpp
@@ -41,82 +41,6 @@ namespace UVEditorSeamToolLocals
 	
 	const FString& ExistingSeamsID(TEXT("SeamLineSet"));
 
-	constexpr double SeamEuclideanDistanceMultiplier = 0.01;
-
-	typedef TFunction<double(int32 FromVID, int32 ToVID, int32 SeedVID, double EuclideanDistance)> FPathDistanceFunc;
-
-	FPathDistanceFunc GetUVDistanceFunction(const UUVEditorToolMeshInput& Target, bool bUnwrapMeshSource, bool bTrackBoundaryEdges)
-	{
-		if (!bTrackBoundaryEdges)
-		{
-			return [&Target](int32 FromVID, int32 ToVID, int32 SeedVID, double EuclideanDistance)
-			{
-				return EuclideanDistance;
-			};
-		}
-		else
-		{
-			if (bUnwrapMeshSource)
-			{
-				return [&Target](int32 FromVID, int32 ToVID, int32 SeedVID, double EuclideanDistance)
-				{
-					if (Target.UnwrapCanonical.Get()->IsBoundaryEdge(Target.UnwrapCanonical.Get()->FindEdge(FromVID, ToVID)))
-					{
-						return EuclideanDistance * SeamEuclideanDistanceMultiplier;
-					}
-					return EuclideanDistance;
-				};
-			}
-			else
-			{
-				return [&Target](int32 FromVID, int32 ToVID, int32 SeedVID, double EuclideanDistance)
-				{					
-					const FDynamicMeshUVOverlay* CanonicalOverlay = Target.AppliedCanonical->Attributes()->GetUVLayer(Target.UVLayerIndex);
-					if (CanonicalOverlay->IsSeamEdge(Target.AppliedCanonical->FindEdge(FromVID, ToVID)))
-					{
-						return EuclideanDistance * SeamEuclideanDistanceMultiplier;
-					}
-					return EuclideanDistance;
-				};
-			}
-		}
-	}
-
-	void GetVidPath(FDynamicMesh3* Mesh, const TArray<int32>& StartVids, int32 EndVid, TArray<int32>& VidPathOut, FPathDistanceFunc PathDistanceFunc = nullptr)
-	{
-		VidPathOut.Reset();
-		if (StartVids.Contains(EndVid))
-		{
-			VidPathOut.Add(EndVid);
-			return;
-		}
-
-		UE::Geometry::TMeshDijkstra<FDynamicMesh3> PathFinder(Mesh);
-		TArray<TMeshDijkstra<FDynamicMesh3>::FSeedPoint> SeedPoints;
-		for (int32 StartVid : StartVids)
-		{
-			SeedPoints.Add({ StartVid, StartVid, 0 });
-		}
-
-		if (PathDistanceFunc)
-		{
-			PathFinder.bEnableDistanceWeighting = true;
-			PathFinder.GetWeightedDistanceFunc = PathDistanceFunc;
-		}
-
-		if (PathFinder.ComputeToTargetPoint(SeedPoints, EndVid))
-		{
-			PathFinder.FindPathToNearestSeed(EndVid, VidPathOut);
-		}
-		Algo::Reverse(VidPathOut);
-	}
-
-	// Like the other GetVidPath, just starts from a specific vert rather than from multiple
-	void GetVidPath(FDynamicMesh3* Mesh, int32 StartVid, int32 EndVid, TArray<int32>& VidPathOut, FPathDistanceFunc PathDistanceFunc = nullptr)
-	{
-		GetVidPath(Mesh, TArray<int32> {StartVid}, EndVid, VidPathOut, PathDistanceFunc);
-	}
-
 	void AddDisplayedPoints(UUVEditorToolMeshInput* InputObject,
 		UPointSetComponent* UnwrapPointSet, UPointSetComponent* AppliedPointSet,
 		int32 AppliedVid, const TArray<int32>& UnwrapVids, const FColor& Color, float DepthBias)
@@ -344,6 +268,16 @@ void UUVEditorSeamTool::Setup()
 
 	UInteractiveTool::Setup();
 
+	// We cache the bounding box max dimensions because we use these to adjust our
+	//  path similarity metric when balancing against path length.
+	UnwrapMaxDims.SetNum(Targets.Num());
+	AppliedMeshMaxDims.SetNum(Targets.Num());
+	for (int32 i = 0; i < Targets.Num(); ++i)
+	{
+		UnwrapMaxDims[i] = Targets[i]->UnwrapCanonical->GetBounds(true).MaxDim();
+		AppliedMeshMaxDims[i] = Targets[i]->AppliedCanonical->GetBounds(true).MaxDim();
+	}
+
 	Settings = NewObject<UUVEditorSeamToolProperties>();
 	Settings->RestoreProperties(this);
 	Settings->WatchProperty(Settings->Mode, [&](EUVEditorSeamMode) { OnSeamModeChanged(); });
@@ -745,14 +679,14 @@ void UUVEditorSeamTool::UpdateHover()
 		{
 			TArray<int32> LastLockedUnwrapVids;
 			Targets[HoverMeshIndex]->AppliedVidToUnwrapVids(LastLockedAppliedVid, LastLockedUnwrapVids);
-			GetVidPath(Targets[HoverMeshIndex]->UnwrapCanonical.Get(), LastLockedUnwrapVids, HoverVid, NewPathVids,
-				       GetUVDistanceFunction(*Targets[HoverMeshIndex], bHoverVidIsFromUnwrap, IsInJoinMode()));
+			GetVidPath(*Targets[HoverMeshIndex], bHoverVidIsFromUnwrap, 
+				LastLockedUnwrapVids, HoverVid, NewPathVids);
 		}
 		else
 		{
 			ensure(LastLockedAppliedVid != IndexConstants::InvalidID);
-			GetVidPath(Targets[HoverMeshIndex]->AppliedCanonical.Get(), LastLockedAppliedVid, HoverVid, NewPathVids,
-				       GetUVDistanceFunction(*Targets[HoverMeshIndex], bHoverVidIsFromUnwrap, IsInJoinMode()));
+			GetVidPath(*Targets[HoverMeshIndex], bHoverVidIsFromUnwrap, 
+				{ LastLockedAppliedVid }, HoverVid, NewPathVids);
 		}
 	}
 
@@ -873,8 +807,7 @@ void UUVEditorSeamTool::ApplyClick()
 	{
 		TArray<int32> LastLockedUnwrapVids;
 		ClickedTarget->AppliedVidToUnwrapVids(LastLockedAppliedVid, LastLockedUnwrapVids);
-		GetVidPath(ClickedTarget->UnwrapCanonical.Get(), LastLockedUnwrapVids, ClickedVid, VidPath,
-			       GetUVDistanceFunction(*ClickedTarget, bClickWasInUnwrap, IsInJoinMode()));
+		GetVidPath(*ClickedTarget, bClickWasInUnwrap, LastLockedUnwrapVids, ClickedVid, VidPath);
 		for (int32 Vid : VidPath)
 		{
 			TempStorage.Add(ClickedTarget->UnwrapVidToAppliedVid(Vid));
@@ -883,8 +816,7 @@ void UUVEditorSeamTool::ApplyClick()
 	}
 	else
 	{
-		GetVidPath(ClickedTarget->AppliedCanonical.Get(), LastLockedAppliedVid, ClickedVid, VidPath,
-			       GetUVDistanceFunction(*ClickedTarget, bClickWasInUnwrap, IsInJoinMode()));
+		GetVidPath(*ClickedTarget, bClickWasInUnwrap, { LastLockedAppliedVid }, ClickedVid, VidPath);
 	}
 
 	if (VidPath.Num() == 0)
@@ -1365,5 +1297,114 @@ void UUVEditorSeamTool::RecordAnalytics()
 	}
 }
 
+void UUVEditorSeamTool::GetVidPath(const UUVEditorToolMeshInput& Target, bool bUnwrapMeshSource, 
+	const TArray<int32>& StartVids, int32 EndVid, TArray<int32>& VidPathOut)
+{
+	VidPathOut.Reset();
+	if (StartVids.Contains(EndVid))
+	{
+		VidPathOut.Add(EndVid);
+		return;
+	}
+
+	FDynamicMesh3* Mesh = bUnwrapMeshSource ? Target.UnwrapCanonical.Get() : Target.AppliedCanonical.Get();
+
+	UE::Geometry::TMeshDijkstra<FDynamicMesh3> PathFinder(Mesh);
+	TArray<TMeshDijkstra<FDynamicMesh3>::FSeedPoint> SeedPoints;
+	for (int32 StartVid : StartVids)
+	{
+		SeedPoints.Add({ StartVid, StartVid, 0 });
+	}
+
+	// See if we need a custom distance metric
+	if (IsInJoinMode() || Settings->PathSimilarityWeight > 0)
+	{
+		PathFinder.bEnableDistanceWeighting = true;
+
+		// We'll replace this with one that includes a similarity metric if the weight is nonzero
+		TUniqueFunction<double(int32 FromVID, int32 ToVID, int32 SeedVID, double EuclideanDistance)> PathMetric = 
+			[](int32 FromVid, int32 ToVid, int32 SeedVid, double EuclideanDistance) 
+		{ 
+			return EuclideanDistance; 
+		};
+
+		double MeshMaxDim = bUnwrapMeshSource ? UnwrapMaxDims[Target.AssetID] : AppliedMeshMaxDims[Target.AssetID];
+		if (Settings->PathSimilarityWeight > 0 && MeshMaxDim > 0)
+		{
+			// Since we have multiple potential sources, we have multiple direct lines from source to destination
+			//  that we can try to measure similarity against.
+			TArray<FLine3d> DirectLines;
+			FVector3d EndPosition = Mesh->GetVertex(EndVid);
+			for (int32 StartVid : StartVids)
+			{
+				DirectLines.Add(FLine3d::FromPoints(Mesh->GetVertex(StartVid), EndPosition));
+			}
+
+			// Replace our metric
+			PathMetric = [PathSimilarityWeight = Settings->PathSimilarityWeight, Mesh, MeshMaxDim, Lines = MoveTemp(DirectLines)]
+				(int32 FromVid, int32 ToVid, int32 SeedVid, double EuclideanDistance)
+			{
+				// Get the lowest integrated squared distance from the lines
+				double SimilarityMetric = TNumericLimits<double>::Max();
+				for (const FLine3d& Line : Lines)
+				{
+					SimilarityMetric = FMath::Min(SimilarityMetric, SquaredDistanceFromLineIntegratedAlongSegment(
+						Line, FSegment3d(Mesh->GetVertex(FromVid), Mesh->GetVertex(ToVid))));
+				}
+
+				// We want the relative path weights to not change if the mesh is uniformly scaled by a positive scalar s.
+				//  Our similarity metric is proportional to s^3, so divide by max dim squared to make it proportional
+				//  to s like EuclideanDistance, so that the final result is just scaled by s. We don't do cube root 
+				//  because our adjustment needs to be distributive, to avoid being sensitive to path tesselation.
+				SimilarityMetric /= (MeshMaxDim * MeshMaxDim);
+
+				return EuclideanDistance + PathSimilarityWeight * SimilarityMetric;
+			};
+		}
+
+		// If we're joining seams, put a seam discount on whatever our path metric is.
+		if (IsInJoinMode())
+		{
+			constexpr double SeamEuclideanDistanceMultiplier = 0.01;
+			if (bUnwrapMeshSource)
+			{
+				PathFinder.GetWeightedDistanceFunc = [Mesh, PathMetric = MoveTemp(PathMetric)](int32 FromVid, int32 ToVid, int32 SeedVid, double EuclideanDistance)
+				{
+					double Metric = PathMetric(FromVid, ToVid, SeedVid, EuclideanDistance);
+					if (Mesh->IsBoundaryEdge(Mesh->FindEdge(FromVid, ToVid)))
+					{
+						return Metric * SeamEuclideanDistanceMultiplier;
+					}
+					return Metric;
+				};
+			}
+			else
+			{
+				const FDynamicMeshUVOverlay* Overlay = Mesh->Attributes()->GetUVLayer(Target.UVLayerIndex);
+				PathFinder.GetWeightedDistanceFunc = [Overlay, Mesh, PathMetric = MoveTemp(PathMetric)](int32 FromVid, int32 ToVid, int32 SeedVid, double EuclideanDistance)
+				{
+					double Metric = PathMetric(FromVid, ToVid, SeedVid, EuclideanDistance);
+					
+					if (Overlay->IsSeamEdge(Mesh->FindEdge(FromVid, ToVid)))
+					{
+						return Metric * SeamEuclideanDistanceMultiplier;
+					}
+					return Metric;
+				};
+			}
+		}
+		else
+		{
+			PathFinder.GetWeightedDistanceFunc = MoveTemp(PathMetric);
+		}
+	}
+
+	if (PathFinder.ComputeToTargetPoint(SeedPoints, EndVid))
+	{
+		PathFinder.FindPathToNearestSeed(EndVid, VidPathOut);
+	}
+	Algo::Reverse(VidPathOut);
+}
+
 #undef LOCTEXT_NAMESPACE
 
diff --git a/Engine/Plugins/Editor/UVEditor/Source/UVEditorTools/Public/ToolTargets/UVEditorToolMeshInput.h b/Engine/Plugins/Editor/UVEditor/Source/UVEditorTools/Public/ToolTargets/UVEditorToolMeshInput.h
--- a/Engine/Plugins/Editor/UVEditor/Source/UVEditorTools/Public/ToolTargets/UVEditorToolMeshInput.h
+++ b/Engine/Plugins/Editor/UVEditor/Source/UVEditorTools/Public/ToolTargets/UVEditorToolMeshInput.h
@@ -118,7 +118,7 @@ public:
 
 	// Additional needed information
 	TObjectPtr<UToolTarget> SourceTarget = nullptr;
-	int32 AssetID = -1;
+	int32 AssetID = -1; // Index into targets array
 	int32 UVLayerIndex = 0;
 
 	// Mappings used for generating and baking back the unwrap.
diff --git a/Engine/Plugins/Editor/UVEditor/Source/UVEditorTools/Public/UVEditorSeamTool.h b/Engine/Plugins/Editor/UVEditor/Source/UVEditorTools/Public/UVEditorSeamTool.h
--- a/Engine/Plugins/Editor/UVEditor/Source/UVEditorTools/Public/UVEditorSeamTool.h
+++ b/Engine/Plugins/Editor/UVEditor/Source/UVEditorTools/Public/UVEditorSeamTool.h
@@ -42,6 +42,16 @@ public:
 
 	UPROPERTY(EditAnywhere, Category = Options)
 	EUVEditorSeamMode Mode = EUVEditorSeamMode::Cut;
+
+	/**
+	 * Setting this above 0 will include a measure of path similarity to seam transfer, so that among
+	 *  similarly short paths, we pick one that lies closer to the edge. Useful in cases where the path
+	 *  is on the wrong diagonal to the triangulation, because it prefers a closely zigzagging path over
+	 *  a wider "up and over" path that has similar length. If set to 0, only path length is used.
+	 */
+	UPROPERTY(EditAnywhere, Category = Options, AdvancedDisplay, meta = (
+		ClampMin = 0, UIMax = 1000))
+	double PathSimilarityWeight = 200;
 };
 
 UCLASS()
@@ -185,4 +195,11 @@ protected:
 	UE::Geometry::UVEditorAnalytics::FTargetAnalytics InputTargetAnalytics;
 	FDateTime ToolStartTimeAnalytics;
 	void RecordAnalytics();
+
+private:
+	void GetVidPath(const UUVEditorToolMeshInput& Target, bool bUnwrapMeshSource,
+		const TArray<int32>& StartVids, int32 EndVid, TArray<int32>& VidPathOut);
+
+	TArray<double> UnwrapMaxDims;
+	TArray<double> AppliedMeshMaxDims;
 };
diff --git a/Engine/Plugins/Experimental/MeshModelingToolsetExp/Source/MeshModelingToolsExp/Private/SeamSculptTool.cpp b/Engine/Plugins/Experimental/MeshModelingToolsetExp/Source/MeshModelingToolsExp/Private/SeamSculptTool.cpp
--- a/Engine/Plugins/Experimental/MeshModelingToolsetExp/Source/MeshModelingToolsExp/Private/SeamSculptTool.cpp
+++ b/Engine/Plugins/Experimental/MeshModelingToolsetExp/Source/MeshModelingToolsExp/Private/SeamSculptTool.cpp
@@ -79,7 +79,9 @@ void USeamSculptTool::Setup()
 	MeshTransform = FTransform3d(TargetComponent->GetWorldTransform());
 	InputMesh = MakeShared<FDynamicMesh3, ESPMode::ThreadSafe>(*PreviewMesh->GetMesh());
 	FMeshNormals::QuickComputeVertexNormals(*InputMesh);
-	NormalOffset = InputMesh->GetBounds(true).MinDim() * 0.001;
+	FAxisAlignedBox3d MeshBounds = InputMesh->GetBounds(true);
+	NormalOffset = MeshBounds.MinDim() * 0.001;
+	MeshMaxDim = MeshBounds.MaxDim();
 
 	// disable shadows
 	//PreviewMesh->GetRootComponent()->bCastDynamicShadow = false;
@@ -231,6 +233,27 @@ void USeamSculptTool::UpdateCurrentDrawPath()
 
 	const FDynamicMesh3* Mesh = PreviewMesh->GetPreviewDynamicMesh();
 	TMeshDijkstra<FDynamicMesh3> PathFinder(Mesh);
+
+	if (Settings->PathSimilarityWeight > 0 && MeshMaxDim > 0)
+	{
+		PathFinder.bEnableDistanceWeighting = true;
+
+		FLine3d DirectLine = FLine3d::FromPoints(Mesh->GetVertex(DrawPathStartVertex), Mesh->GetVertex(CurrentSnapVertex));
+		PathFinder.GetWeightedDistanceFunc = [this, Mesh, DirectLine](int32 FromVid, int32 ToVid, int32 SeedVid, double EuclideanDistance)
+		{
+			double SimilarityMetric = SquaredDistanceFromLineIntegratedAlongSegment(
+				DirectLine, FSegment3d(Mesh->GetVertex(FromVid), Mesh->GetVertex(ToVid)));
+
+			// We want the relative path weights to not change if the mesh is uniformly scaled by a positive scalar s.
+			//  Our similarity metric is proportional to s^3, so divide by max dim squared to make it proportional
+			//  to s like EuclideanDistance, so that the final result is just scaled by s. We don't do cube root 
+			//  because our adjustment needs to be distributive, to avoid being sensitive to path tesselation.
+			SimilarityMetric /= (MeshMaxDim * MeshMaxDim);
+
+			return EuclideanDistance + Settings->PathSimilarityWeight * SimilarityMetric;
+		};
+	}
+
 	TArray<TMeshDijkstra<FDynamicMesh3>::FSeedPoint> SeedPoints;
 	SeedPoints.Add({ DrawPathStartVertex, DrawPathStartVertex, 0 });
 	if ( PathFinder.ComputeToTargetPoint(SeedPoints, CurrentSnapVertex) )
diff --git a/Engine/Plugins/Experimental/MeshModelingToolsetExp/Source/MeshModelingToolsExp/Public/SeamSculptTool.h b/Engine/Plugins/Experimental/MeshModelingToolsetExp/Source/MeshModelingToolsExp/Public/SeamSculptTool.h
--- a/Engine/Plugins/Experimental/MeshModelingToolsetExp/Source/MeshModelingToolsExp/Public/SeamSculptTool.h
+++ b/Engine/Plugins/Experimental/MeshModelingToolsetExp/Source/MeshModelingToolsExp/Public/SeamSculptTool.h
@@ -40,6 +40,16 @@ public:
 
 	UPROPERTY(EditAnywhere, Category = Options)
 	bool bHitBackFaces = false;
+
+	/**
+	 * Setting this above 0 will include a measure of path similarity to seam transfer, so that among
+	 *  similarly short paths, we pick one that lies closer to the edge. Useful in cases where the path
+	 *  is on the wrong diagonal to the triangulation, because it prefers a closely zigzagging path over
+	 *  a wider "up and over" path that has similar length. If set to 0, only path length is used.
+	 */
+	UPROPERTY(EditAnywhere, Category = Options, AdvancedDisplay, meta = (
+		ClampMin = 0, UIMax = 1000))
+	double PathSimilarityWeight = 200;
 };
 
 
@@ -122,4 +132,6 @@ protected:
 
 	UWorld* TargetWorld;
 
+private:
+	double MeshMaxDim = 0;
 };
\ No newline at end of file
