Modeling Tools: Cleans and fixes bugs in UV Snapshot tool, adds icon

diff --git a/Engine/Plugins/Editor/UVEditor/Source/UVEditor/Private/UVEditorMode.cpp b/Engine/Plugins/Editor/UVEditor/Source/UVEditor/Private/UVEditorMode.cpp
--- a/Engine/Plugins/Editor/UVEditor/Source/UVEditor/Private/UVEditorMode.cpp
+++ b/Engine/Plugins/Editor/UVEditor/Source/UVEditor/Private/UVEditorMode.cpp
@@ -1068,15 +1068,6 @@ void UUVEditorMode::InitializeAssetEditorContexts(UContextObjectStore& ContextSt
 	const TArray<TObjectPtr<UObject>>& AssetsIn, const TArray<FTransform>& TransformsIn,
 	FEditorViewportClient& LivePreviewViewportClient, FAssetEditorModeManager& LivePreviewModeManager,
 	UUVToolViewportButtonsAPI& ViewportButtonsAPI, UUVTool2DViewportAPI& UVTool2DViewportAPI)
-{
-	UInteractiveToolsContext* InteractiveToolsContext = Cast<UInteractiveToolsContext>(ContextStore.GetOuter());
-	InitializeAssetEditorContexts(ContextStore, AssetsIn, TransformsIn, LivePreviewViewportClient, LivePreviewModeManager,
-		ViewportButtonsAPI, UVTool2DViewportAPI, *InteractiveToolsContext);
-}
-void UUVEditorMode::InitializeAssetEditorContexts(UContextObjectStore& ContextStore,
-	const TArray<TObjectPtr<UObject>>& AssetsIn, const TArray<FTransform>& TransformsIn,
-	FEditorViewportClient& LivePreviewViewportClient, FAssetEditorModeManager& LivePreviewModeManager,
-	UUVToolViewportButtonsAPI& ViewportButtonsAPI, UUVTool2DViewportAPI& UVTool2DViewportAPI, UInteractiveToolsContext& ToolsContext)
 {
 	using namespace UVEditorModeLocals;
 
@@ -1091,7 +1082,7 @@ void UUVEditorMode::InitializeAssetEditorContexts(UContextObjectStore& ContextSt
 	UEditorModelingObjectsCreationAPI* ModelingObjectsCreationAPI = ContextStore.FindContext<UEditorModelingObjectsCreationAPI>();
 	if (!ModelingObjectsCreationAPI)
 	{
-		ModelingObjectsCreationAPI = NewObject<UEditorModelingObjectsCreationAPI>(&ToolsContext);
+		ModelingObjectsCreationAPI = NewObject<UEditorModelingObjectsCreationAPI>();
 		ContextStore.AddContextObject(ModelingObjectsCreationAPI);
 	}
 	
diff --git a/Engine/Plugins/Editor/UVEditor/Source/UVEditor/Private/UVEditorToolkit.cpp b/Engine/Plugins/Editor/UVEditor/Source/UVEditor/Private/UVEditorToolkit.cpp
--- a/Engine/Plugins/Editor/UVEditor/Source/UVEditor/Private/UVEditorToolkit.cpp
+++ b/Engine/Plugins/Editor/UVEditor/Source/UVEditor/Private/UVEditorToolkit.cpp
@@ -453,7 +453,7 @@ void FUVEditorToolkit::PostInitAssetEditor()
 	// the context store on Enter() to function properly.
 	UUVEditorMode::InitializeAssetEditorContexts(*EditorModeManager->GetInteractiveToolsContext()->ContextObjectStore, 
 		ObjectsToEdit, ObjectTransforms, *LivePreviewViewportClient, *LivePreviewEditorModeManager, 
-		*ViewportButtonsAPI, *UVTool2DViewportAPI, *EditorModeManager->GetInteractiveToolsContext());
+		*ViewportButtonsAPI, *UVTool2DViewportAPI);
 
 	// Currently, aside from setting up all the UI elements, the toolkit also kicks off the UV
 	// editor mode, which is the mode that the editor always works in (things are packaged into
diff --git a/Engine/Plugins/Editor/UVEditor/Source/UVEditor/Public/UVEditorMode.h b/Engine/Plugins/Editor/UVEditor/Source/UVEditor/Public/UVEditorMode.h
--- a/Engine/Plugins/Editor/UVEditor/Source/UVEditor/Public/UVEditorMode.h
+++ b/Engine/Plugins/Editor/UVEditor/Source/UVEditor/Public/UVEditorMode.h
@@ -176,15 +176,10 @@ public:
 	/**
 	 * Called by an asset editor so that a created instance of the mode has all the data it needs on Enter() to initialize itself.
 	 */
-	UE_DEPRECATED(5.5, "Initialization of contexts now requires UInteractiveToolsContext to be passed in directly. Please update code accordingly")
 	static void InitializeAssetEditorContexts(UContextObjectStore& ContextStore,
 		const TArray<TObjectPtr<UObject>>& AssetsIn, const TArray<FTransform>& TransformsIn,
 		FEditorViewportClient& LivePreviewViewportClient, FAssetEditorModeManager& LivePreviewModeManager,
 		UUVToolViewportButtonsAPI& ViewportButtonsAPI, UUVTool2DViewportAPI& UVTool2DViewportAPI);
-	static void InitializeAssetEditorContexts(UContextObjectStore& ContextStore,
-		const TArray<TObjectPtr<UObject>>& AssetsIn, const TArray<FTransform>& TransformsIn,
-		FEditorViewportClient& LivePreviewViewportClient, FAssetEditorModeManager& LivePreviewModeManager,
-		UUVToolViewportButtonsAPI& ViewportButtonsAPI, UUVTool2DViewportAPI& UVTool2DViewportAPI, UInteractiveToolsContext& ToolsContext);
 
 	// Public for use by undo/redo. Otherwise should use RequestUVChannelChange
 	void ChangeInputObjectLayer(int32 AssetID, int32 NewLayerIndex);
diff --git a/Engine/Plugins/Editor/UVEditor/Source/UVEditorTools/Private/UVEditorStyle.cpp b/Engine/Plugins/Editor/UVEditor/Source/UVEditorTools/Private/UVEditorStyle.cpp
--- a/Engine/Plugins/Editor/UVEditor/Source/UVEditorTools/Private/UVEditorStyle.cpp
+++ b/Engine/Plugins/Editor/UVEditor/Source/UVEditorTools/Private/UVEditorStyle.cpp
@@ -35,7 +35,8 @@ FUVEditorStyle::FUVEditorStyle()
 	Set("UVEditor.BeginTexelDensityTool",       new IMAGE_BRUSH_SVG("TexelDensity", ToolbarIconSize));
 	Set("UVEditor.BeginBrushSelectTool",        new IMAGE_BRUSH("MeshSelect_40x", ToolbarIconSize));
 	Set("UVEditor.BeginBrushSelectTool.Small",  new IMAGE_BRUSH("MeshSelect_40x", ToolbarIconSize));
-
+	Set("UVEditor.BeginUVSnapshotTool",		    new IMAGE_BRUSH_SVG("UVSnapshot", ToolbarIconSize));
+	
 	// Select tool actions
 	Set("UVEditor.SewAction", new IMAGE_BRUSH_SVG("UVSew", ToolbarIconSize));
 	Set("UVEditor.SplitAction", new IMAGE_BRUSH_SVG("UVCut", ToolbarIconSize));
diff --git a/Engine/Plugins/Editor/UVEditor/Source/UVEditorTools/Private/UVEditorUVSnapshotTool.cpp b/Engine/Plugins/Editor/UVEditor/Source/UVEditorTools/Private/UVEditorUVSnapshotTool.cpp
--- a/Engine/Plugins/Editor/UVEditor/Source/UVEditorTools/Private/UVEditorUVSnapshotTool.cpp
+++ b/Engine/Plugins/Editor/UVEditor/Source/UVEditorTools/Private/UVEditorUVSnapshotTool.cpp
@@ -10,7 +10,6 @@
 #include "ModelingObjectsCreationAPI.h"
 #include "Materials/Material.h"
 #include "Materials/MaterialInstanceDynamic.h"
-#include "TargetInterfaces/StaticMeshBackedTarget.h"
 #include "UVEditorUXSettings.h"
 #include "Drawing/MeshElementsVisualizer.h"
 
@@ -25,6 +24,7 @@ using namespace UE::Geometry;
  */
 bool UUVEditorUVSnapshotToolBuilder::CanBuildTool(const FToolBuilderState& SceneState) const
 {
+	// ensure only one target is handled
 	return Targets && Targets->Num() == 1;
 }
 UInteractiveTool* UUVEditorUVSnapshotToolBuilder::BuildTool(const FToolBuilderState& SceneState) const
@@ -44,9 +44,13 @@ public:
 	// General bake settings
 	TSharedPtr<UE::Geometry::FDynamicMesh3, ESPMode::ThreadSafe> BaseMesh;
 	TUniquePtr<UE::Geometry::FMeshMapBaker> Baker;
-	UUVEditorUVSnapshotTool::FUVSnapshotBakeSettings BakeSettings;
 	TSharedPtr<UE::Geometry::FDynamicMeshAABBTree3, ESPMode::ThreadSafe> DetailSpatial;
 
+	TSharedPtr<FMeshUVShellMapEvaluator> UVShellEval = MakeShared<FMeshUVShellMapEvaluator>();
+
+	UE::Geometry::FImageDimensions BakerDimensions;
+	int BakerSamplesPerPixel = 4;
+
 	virtual void CalculateResult(FProgressCancel* Progress) override
 	{
 		Baker = MakeUnique<FMeshMapBaker>();
@@ -54,19 +58,15 @@ public:
 			return Progress && Progress->Cancelled();
 		};
 		Baker->SetTargetMesh(BaseMesh.Get());
-		Baker->SetTargetMeshUVLayer(BakeSettings.UVLayer);
-		Baker->SetDimensions(BakeSettings.Dimensions);
-		Baker->SetSamplesPerPixel(BakeSettings.SamplesPerPixel);
+		Baker->SetTargetMeshUVLayer(UVShellEval->UVLayer);
+		Baker->SetDimensions(BakerDimensions);
+		Baker->SetSamplesPerPixel(BakerSamplesPerPixel);
 		
 		FMeshBakerDynamicMeshSampler DetailSampler(BaseMesh.Get(), DetailSpatial.Get(), nullptr);
 		Baker->SetDetailSampler(&DetailSampler);
-		const TSharedPtr<FMeshUVShellMapEvaluator> UVShellEval = MakeShared<FMeshUVShellMapEvaluator>();
-		UVShellEval->TexelSize = BakeSettings.Dimensions.GetTexelSize();
-		UVShellEval->UVLayer = BakeSettings.UVLayer;
-		UVShellEval->WireframeThickness = BakeSettings.WireframeThickness;
-		UVShellEval->WireframeColor = BakeSettings.WireframeColor;
-		UVShellEval->ShellColor = BakeSettings.ShellColor;
-		UVShellEval->BackgroundColor = BakeSettings.BackgroundColor;
+
+		UVShellEval->TexelSize = BakerDimensions.GetTexelSize();
+
 		Baker->AddEvaluator(UVShellEval);
 		Baker->Bake();
 		SetResult(MoveTemp(Baker));
@@ -88,20 +88,15 @@ void UUVEditorUVSnapshotTool::Setup()
 	// retrieve whatever UV Layer is currently being displayed in UV Editor
 	UVShellSettings->UVLayer = Target->UVLayerIndex;
 	
-	UVShellSettings->WatchProperty(UVShellSettings->SamplesPerPixel, [this](EBakeTextureSamplesPerPixel) { OpState |= EUVSnapshotBakeOpState::Evaluate; });
-	UVShellSettings->WatchProperty(UVShellSettings->Resolution, [this](EBakeTextureResolution) { OpState |= EUVSnapshotBakeOpState::Evaluate; });
-	UVShellSettings->WatchProperty(UVShellSettings->UVLayer, [this](int) { OpState |= EUVSnapshotBakeOpState::Evaluate; });
-	UVShellSettings->WatchProperty(UVShellSettings->WireframeThickness, [this](float) { OpState |= EUVSnapshotBakeOpState::Evaluate; });
-	UVShellSettings->WatchProperty(UVShellSettings->WireframeColor, [this](FLinearColor) { OpState |= EUVSnapshotBakeOpState::Evaluate; });
-	UVShellSettings->WatchProperty(UVShellSettings->ShellColor, [this](FLinearColor) { OpState |= EUVSnapshotBakeOpState::Evaluate; });
-	UVShellSettings->WatchProperty(UVShellSettings->BackgroundColor, [this](FLinearColor) { OpState |= EUVSnapshotBakeOpState::Evaluate; });
+	UVShellSettings->WatchProperty(UVShellSettings->SamplesPerPixel, [this](EBakeTextureSamplesPerPixel) { Compute->InvalidateResult(); });
+	UVShellSettings->WatchProperty(UVShellSettings->Resolution, [this](EBakeTextureResolution) { Compute->InvalidateResult(); });
+	UVShellSettings->WatchProperty(UVShellSettings->UVLayer, [this](int) { Compute->InvalidateResult(); });
+	UVShellSettings->WatchProperty(UVShellSettings->WireframeThickness, [this](float) { Compute->InvalidateResult(); });
+	UVShellSettings->WatchProperty(UVShellSettings->WireframeColor, [this](FLinearColor) { Compute->InvalidateResult(); });
+	UVShellSettings->WatchProperty(UVShellSettings->ShellColor, [this](FLinearColor) { Compute->InvalidateResult(); });
+	UVShellSettings->WatchProperty(UVShellSettings->BackgroundColor, [this](FLinearColor) {Compute->InvalidateResult(); });
+	UVShellSettings->Result = nullptr;
 	SetToolPropertySourceEnabled(UVShellSettings, true);
-	
-	ResultSettings = NewObject<UUVEditorBakeUVShellResultProperties>(this);
-	ResultSettings->RestoreProperties(this);
-	AddToolPropertySource(ResultSettings);
-	SetToolPropertySourceEnabled(ResultSettings, true);
-	ResultSettings->Result = nullptr;
 
 	// set up the detail mesh & spatial
 	FDynamicMesh3 DetailMeshGet;
@@ -111,18 +106,21 @@ void UUVEditorUVSnapshotTool::Setup()
 	DetailSpatial = MakeShared<FDynamicMeshAABBTree3, ESPMode::ThreadSafe>();
 	DetailSpatial->SetMesh(DetailMesh.Get(), true);
 
-	// mark for evaluation
-	OpState |= EUVSnapshotBakeOpState::Evaluate;
-
-	CachedBakeSettings = FUVSnapshotBakeSettings();
-
 	// set up UPreviewGeometry for visualization in Unwrap viewport
 	PreviewGeoBackgroundQuad = NewObject<UPreviewGeometry>(this);
 	PreviewGeoBackgroundQuad->CreateInWorld(Target->UnwrapPreview->GetWorld(), FTransform::Identity);
 	PreviewGeoBackgroundQuad->AddTriangleSet("UVShellMap");
-	PreviewGeoBackgroundQuad->SetAllVisible(true);
+	PreviewGeoBackgroundQuad->SetAllVisible(false);
+
+	SetUpPreviewQuad();
+
+	// Initialize background compute
+	Compute = MakeUnique<TGenericDataBackgroundCompute<FMeshMapBaker>>();
+	Compute->Setup(this);
+	Compute->OnResultUpdated.AddLambda([this](const TUniquePtr<FMeshMapBaker>& NewResult) { OnMapUpdated(NewResult); });
+
+	Compute->InvalidateResult();
 
-	// set up op factory
 	SetToolDisplayName(LOCTEXT("ToolNameLocal", "UV Snapshot"));
 	GetToolManager()->DisplayMessage(
 		LOCTEXT("OnStartUVSnapshotTool", "Export a texture asset of a UV Layout."),
@@ -133,10 +131,7 @@ void UUVEditorUVSnapshotTool::Shutdown(EToolShutdownType ShutdownType)
 {
 	if (ShutdownType == EToolShutdownType::Accept)
 	{
-		const IStaticMeshBackedTarget* StaticMeshTarget = Cast<IStaticMeshBackedTarget>(Target->SourceTarget);
-		UObject* SourceAsset = StaticMeshTarget ? StaticMeshTarget->GetStaticMesh() : nullptr;
-		const UPrimitiveComponent* SourceComponent = UE::ToolTarget::GetTargetComponent(Target->SourceTarget);
-		CreateTextureAsset(ResultSettings->Result, SourceComponent->GetWorld(), SourceAsset);
+		CreateTextureAsset(UVShellSettings->Result);
 	}
 	
 	UVShellSettings->SaveProperties(this);
@@ -150,8 +145,6 @@ void UUVEditorUVSnapshotTool::Shutdown(EToolShutdownType ShutdownType)
 		Compute->Shutdown();
 	}
 
-	Target->AppliedPreview->ClearOpFactory();
-
 	if (PreviewGeoBackgroundQuad)
 	{
 		PreviewGeoBackgroundQuad->Disconnect();
@@ -168,131 +161,68 @@ void UUVEditorUVSnapshotTool::OnTick(float DeltaTime)
 		Compute->Tick(DeltaTime);
 	}
 }
-void UUVEditorUVSnapshotTool::Render(IToolsContextRenderAPI* RenderAPI)
-{
-	UpdateResult();
-}
+
 TUniquePtr<UE::Geometry::TGenericDataOperator<UE::Geometry::FMeshMapBaker>> UUVEditorUVSnapshotTool::MakeNewOperator()
 {
 	TUniquePtr<FMeshUVMapBakerOp> Op = MakeUnique<FMeshUVMapBakerOp>();
 	Op->DetailSpatial = DetailSpatial;
 	Op->BaseMesh = DetailMesh;
-	Op->BakeSettings = CachedBakeSettings;
-	
-	return Op;
-}
-void UUVEditorUVSnapshotTool::UpdateResult()
-{
-	if (OpState == EUVSnapshotBakeOpState::Clean)
-	{
-		return;
-	}
-	// clear warning (ugh)
-	GetToolManager()->DisplayMessage(FText(), EToolMessageLevel::UserWarning);
 
 	const int32 ImageSize = (int32)UVShellSettings->Resolution;
-	const FImageDimensions Dimensions(ImageSize, ImageSize);
-	CachedBakeSettings.Dimensions = Dimensions;
+	const FImageDimensions ImgDimensions(ImageSize, ImageSize);
+	Op->BakerDimensions = ImgDimensions;
+	Op->BakerSamplesPerPixel = (int32)UVShellSettings->SamplesPerPixel;
+
+	Op->UVShellEval->UVLayer = UVShellSettings->UVLayer;
+	Op->UVShellEval->WireframeThickness = UVShellSettings->WireframeThickness;
+	Op->UVShellEval->WireframeColor = UVShellSettings->WireframeColor;
+	Op->UVShellEval->ShellColor = UVShellSettings->ShellColor;
+	Op->UVShellEval->BackgroundColor = UVShellSettings->BackgroundColor;
 	
-	// Clear our invalid bitflag to check again for valid inputs.
-	OpState &= ~EUVSnapshotBakeOpState::Invalid;
-	OpState |= UpdateResult_UVShellMap();
-	
-	// Early exit if op input parameters are invalid.
-	if ((bool)(OpState & EUVSnapshotBakeOpState::Invalid))
-	{
-		InvalidateResults();
-		return;
-	}
-	// This should be the only point of compute invalidation to
-	// minimize synchronization issues.
-	InvalidateCompute();
+	return Op;
 }
+
 void UUVEditorUVSnapshotTool::OnMapUpdated(const TUniquePtr<UE::Geometry::FMeshMapBaker>& NewResult)
 {
 	const FImageDimensions BakeDimensions = NewResult->GetDimensions();
-	constexpr bool bConvertToSRGB = true;
 	FTexture2DBuilder TextureBuilder;
 	TextureBuilder.Initialize(FTexture2DBuilder::ETextureType::Color, BakeDimensions);
-	TextureBuilder.Copy(*NewResult->GetBakeResults(0)[0], bConvertToSRGB);
+	TextureBuilder.Copy(*NewResult->GetBakeResults(0)[0], true);
 	TextureBuilder.Commit(false);
+
 	// Copy image to source data after commit. This will avoid incurring
 	// the cost of hitting the DDC for texture compile while iterating on
 	// bake settings. Since this dirties the texture, the next time the texture
 	// is used after accepting the final texture, the DDC will trigger and
 	// properly recompile the platform data.
-	constexpr bool bConvertSourceToSRGB = true;
-	// default should be ChannelBits8
-	constexpr ETextureSourceFormat SourceDataFormat = TSF_BGRA8;
-	TextureBuilder.CopyImageToSourceData(*NewResult->GetBakeResults(0)[0], SourceDataFormat, bConvertSourceToSRGB);
-	// The CachedUVMap can be thrown out of sync if updated during a background
-	// compute. Validate the computed type against our cached map.
+	constexpr ETextureSourceFormat SourceDataFormat = TSF_BGRA8; // default ChannelBits8
+	TextureBuilder.CopyImageToSourceData(*NewResult->GetBakeResults(0)[0], SourceDataFormat, true);
+	
 	CachedUVMap = TextureBuilder.GetTexture2D();
 	UpdateVisualization();
 	GetToolManager()->PostInvalidation();
-	
-}
-void UUVEditorUVSnapshotTool::InvalidateResults() const
-{
-	ResultSettings->Result = nullptr;
-}
-void UUVEditorUVSnapshotTool::InvalidateCompute()
-{
-	if (!Compute)
-	{
-		// Initialize background compute
-		Compute = MakeUnique<TGenericDataBackgroundCompute<FMeshMapBaker>>();
-		Compute->Setup(this);
-		Compute->OnResultUpdated.AddLambda([this](const TUniquePtr<FMeshMapBaker>& NewResult) { OnMapUpdated(NewResult); });
-	}
-	Compute->InvalidateResult();
-	OpState = EUVSnapshotBakeOpState::Clean;
 }
+
 void UUVEditorUVSnapshotTool::UpdateVisualization() 
 {
-	ResultSettings->Result = CachedUVMap;
-	UpdatePreviewMaterialBasedOnBackground();
-}
-EUVSnapshotBakeOpState UUVEditorUVSnapshotTool::UpdateResult_UVShellMap()
-{
-	EUVSnapshotBakeOpState ResultState = EUVSnapshotBakeOpState::Clean;
-	FUVSnapshotBakeSettings UVShellMapSettings;
-	UVShellMapSettings.UVLayer = UVShellSettings->UVLayer;
-	UVShellMapSettings.WireframeThickness = UVShellSettings->WireframeThickness;
-	UVShellMapSettings.WireframeColor = UVShellSettings->WireframeColor;
-	UVShellMapSettings.ShellColor = UVShellSettings->ShellColor;
-	UVShellMapSettings.BackgroundColor = UVShellSettings->BackgroundColor;
+	UVShellSettings->Result = CachedUVMap;
 
-	const int32 ImageSize = (int32)UVShellSettings->Resolution;
-	const FImageDimensions Dimensions(ImageSize, ImageSize);
-	UVShellMapSettings.Dimensions = Dimensions;
-	UVShellMapSettings.SamplesPerPixel = (int32)UVShellSettings->SamplesPerPixel;
+	UTriangleSetComponent* TriangleSet = PreviewGeoBackgroundQuad->FindTriangleSet("UVShellMap");
+	TriangleSet->SetAllTrianglesMaterial(GetMaterialForQuad());
 
-	// if our settings have changed
-	if (CachedBakeSettings != UVShellMapSettings)
-	{
-		CachedBakeSettings = UVShellMapSettings;
-		ResultState |= EUVSnapshotBakeOpState::Evaluate;
-	}
-	return ResultState;
+	// setting visibility here and not in setup to avoid brief moment between tool activation and initial result computation & display
+	// which showed default 'BackgroundBaseMap_Color' texture
+	PreviewGeoBackgroundQuad->SetAllVisible(true);
 }
-void UUVEditorUVSnapshotTool::UpdatePreviewMaterialBasedOnBackground()
-{
-	// using this material so that we can set a texture
-	UMaterialInstanceDynamic* Mat = UMaterialInstanceDynamic::Create(LoadObject<UMaterial>(nullptr, TEXT("/UVEditor/Materials/UVEditorBackground")), this);
-	Mat->SetTextureParameterValue(TEXT("BackgroundBaseMap_Color"), ResultSettings->Result);
-	Mat->SetScalarParameterValue(TEXT("BackgroundPixelDepthOffset"), FUVEditorUXSettings::BackgroundQuadDepthOffset - 0.5);	
 
+void UUVEditorUVSnapshotTool::SetUpPreviewQuad()
+{
 	// temporarily disable wireframe overlay and unwrap preview; re-enabled on shutdown
 	Target->WireframeDisplay->Settings->bVisible = false;
 	Target->UnwrapPreview->SetVisibility(false);
 
-	// connect to existing Preview Geometry
-	UTriangleSetComponent* TriangleSet = PreviewGeoBackgroundQuad->FindTriangleSet("UVShellMap");
-	TriangleSet->Clear();
-
 	const FVector Normal(0, 0, 1);
-	const FColor BackgroundColor = FColor::Red;
+	const FColor BackgroundColor = FColor::Black;
 
 	// set up rendering of 2 triangles to make one 2d quad
 	FIndex2i UDimBlockToRender = FIndex2i(0,0);
@@ -311,12 +241,28 @@ void UUVEditorUVSnapshotTool::UpdatePreviewMaterialBasedOnBackground()
 	FRenderableTriangleVertex V11 = MakeQuadVert(1, 1);
 	FRenderableTriangleVertex V01 = MakeQuadVert(0,1);
 
+	// connect to existing Preview Geometry
+	UTriangleSetComponent* TriangleSet = PreviewGeoBackgroundQuad->FindTriangleSet("UVShellMap");
+	TriangleSet->Clear();
+
+	UMaterialInstanceDynamic* QuadMaterial = GetMaterialForQuad();
+
 	// add to TriangleSet in PreviewGeometry
-	TriangleSet->AddTriangle(FRenderableTriangle(Mat, V00, V10, V11));
-	TriangleSet->AddTriangle(FRenderableTriangle(Mat, V00, V11, V01));
-	
+	TriangleSet->AddTriangle(FRenderableTriangle(QuadMaterial, V00, V10, V11));
+	TriangleSet->AddTriangle(FRenderableTriangle(QuadMaterial, V00, V11, V01));
+}
+
+UMaterialInstanceDynamic* UUVEditorUVSnapshotTool::GetMaterialForQuad()
+{
+	// using this material so that we can set a texture
+	UMaterialInstanceDynamic* QuadMat = UMaterialInstanceDynamic::Create(LoadObject<UMaterial>(nullptr, TEXT("/UVEditor/Materials/UVEditorBackground")), this);
+	QuadMat->SetTextureParameterValue(TEXT("BackgroundBaseMap_Color"), UVShellSettings->Result);
+	QuadMat->SetScalarParameterValue(TEXT("BackgroundPixelDepthOffset"), FUVEditorUXSettings::BackgroundQuadDepthOffset - 1.0f);
+
+	return QuadMat;
 }
-void UUVEditorUVSnapshotTool::CreateTextureAsset(const TObjectPtr<UTexture2D>& Texture, UWorld* SourceWorld, UObject* SourceAsset) const
+
+void UUVEditorUVSnapshotTool::CreateTextureAsset(const TObjectPtr<UTexture2D>& Texture) const
 {
 	bool bCreatedAssetOK = true;
 	const FString ObjName = UE::ToolTarget::GetTargetActor(Target->SourceTarget)->GetActorNameOrLabel();
@@ -347,8 +293,10 @@ void UUVEditorUVSnapshotTool::CreateTextureAsset(const TObjectPtr<UTexture2D>& T
 			PackageNameOut, AssetNameOut);
 
 		// create the asset
-		bCreatedAssetOK = bCreatedAssetOK && UE::Modeling::CreateTextureObject( GetToolManager(),
-			FCreateTextureObjectParams{ 0, SourceWorld, SourceAsset, NewAssetName, Texture}).IsOK();
+		FCreateTextureObjectParams TexParams;
+		TexParams.FullAssetPath = PackageNameOut;
+		TexParams.GeneratedTransientTexture = Texture;
+		bCreatedAssetOK = bCreatedAssetOK && UE::Modeling::CreateTextureObject( GetToolManager(),FCreateTextureObjectParams(TexParams)).IsOK();
 	}
 		
 	ensure(bCreatedAssetOK);
diff --git a/Engine/Plugins/Editor/UVEditor/Source/UVEditorTools/Private/UVEditorUXSettings.cpp b/Engine/Plugins/Editor/UVEditor/Source/UVEditorTools/Private/UVEditorUXSettings.cpp
--- a/Engine/Plugins/Editor/UVEditor/Source/UVEditorTools/Private/UVEditorUXSettings.cpp
+++ b/Engine/Plugins/Editor/UVEditor/Source/UVEditorTools/Private/UVEditorUXSettings.cpp
@@ -13,26 +13,26 @@ const float FUVEditorUXSettings::CameraNearPlaneProportionZ(0.8); // Top layer,
 // Note: While these are floating point values, they represent percentages and should be separated
 // by at least integer amounts, as they serve double duty in certain cases for translucent primitive
 // sorting order.
-const float FUVEditorUXSettings::ToolLockedPathDepthBias(8.0);
-const float FUVEditorUXSettings::ToolExtendPathDepthBias(8.0);
-const float FUVEditorUXSettings::SewLineDepthOffset(7.0f);
-const float FUVEditorUXSettings::SelectionHoverWireframeDepthBias(6);
-const float FUVEditorUXSettings::SelectionHoverTriangleDepthBias(5);
-const float FUVEditorUXSettings::SelectionWireframeDepthBias(4.0);
-const float FUVEditorUXSettings::SelectionTriangleDepthBias(3.0);
-const float FUVEditorUXSettings::WireframeDepthOffset(2.0);
-const float FUVEditorUXSettings::UnwrapTriangleDepthOffset(1.0);
-
-const float FUVEditorUXSettings::LivePreviewExistingSeamDepthBias(1.0);
+const float FUVEditorUXSettings::ToolLockedPathDepthBias(9.0);
+const float FUVEditorUXSettings::ToolExtendPathDepthBias(9.0);
+const float FUVEditorUXSettings::SewLineDepthOffset(8.0f);
+const float FUVEditorUXSettings::SelectionHoverWireframeDepthBias(7);
+const float FUVEditorUXSettings::SelectionHoverTriangleDepthBias(6);
+const float FUVEditorUXSettings::SelectionWireframeDepthBias(5.0);
+const float FUVEditorUXSettings::SelectionTriangleDepthBias(4.0);
+const float FUVEditorUXSettings::WireframeDepthOffset(3.0);
+const float FUVEditorUXSettings::UnwrapTriangleDepthOffset(2.0);
+
+const float FUVEditorUXSettings::LivePreviewExistingSeamDepthBias(2.0);
 
 // Note: that this offset can only be applied when we use our own background material
 // for a user-supplied texture, and we can't use it for a user-provided material.
 // So for consistency this should stay at zero.
 
-const float FUVEditorUXSettings::BackgroundQuadDepthOffset(0.0); // Bottom layer
+const float FUVEditorUXSettings::BackgroundQuadDepthOffset(1.0); // Bottom layer
 
 // 3D Viewport Depth Offsets
-const float FUVEditorUXSettings::LivePreviewHighlightDepthOffset(0.5);
+const float FUVEditorUXSettings::LivePreviewHighlightDepthOffset(1.5);
 
 // Opacities
 const float FUVEditorUXSettings::UnwrapTriangleOpacity(1.0);
diff --git a/Engine/Plugins/Editor/UVEditor/Source/UVEditorTools/Public/UVEditorUVSnapshotTool.h b/Engine/Plugins/Editor/UVEditor/Source/UVEditorTools/Public/UVEditorUVSnapshotTool.h
--- a/Engine/Plugins/Editor/UVEditor/Source/UVEditorTools/Public/UVEditorUVSnapshotTool.h
+++ b/Engine/Plugins/Editor/UVEditor/Source/UVEditorTools/Public/UVEditorUVSnapshotTool.h
@@ -13,20 +13,7 @@
 
 
 // Forward Declarations
-class UUVEditorBakeUVShellResultProperties;
 class UUVEditorBakeUVShellProperties;
-enum class EBakeMapType;
-
-/**
- * Bake compute state
- */
-enum class EUVSnapshotBakeOpState
-{
-	Clean              = 0,      // No-op - evaluation already launched/complete.
-	Evaluate           = 1 << 0, // Inputs are modified and valid, re-evaluate.
-	Invalid            = 1 << 1  // Inputs are modified and invalid - retry eval until valid.
-};
-ENUM_CLASS_FLAGS(EUVSnapshotBakeOpState);
 
 /**
  * ToolBuilder
@@ -59,7 +46,7 @@ public:
 	virtual void Shutdown(EToolShutdownType ShutdownType) override;
 	
 	virtual void OnTick(float DeltaTime) override;
-	virtual void Render(IToolsContextRenderAPI* RenderAPI) override;
+	virtual void Render(IToolsContextRenderAPI* RenderAPI) override { };
 	virtual bool HasCancel() const override { return true; }
 	virtual bool HasAccept() const override { return true; }
 	virtual bool CanAccept() const override { return true; }
@@ -76,7 +63,7 @@ public:
 	{
 		Target = TargetIn;
 	}
-protected:
+private:
 	//
 	// Mesh input to UV Editor
 	//
@@ -84,12 +71,10 @@ protected:
 	TObjectPtr<UUVEditorToolMeshInput> Target;
 
 	//
-	// Property sets for bake and result
+	// Property set for bake and result
 	//
 	UPROPERTY()
 	TObjectPtr<UUVEditorBakeUVShellProperties> UVShellSettings;
-	UPROPERTY()
-	TObjectPtr<UUVEditorBakeUVShellResultProperties> ResultSettings;
 
 	//
 	// Preview Geometry for display in Unwrapped viewport
@@ -102,44 +87,10 @@ protected:
 	//
 	TUniquePtr<TGenericDataBackgroundCompute<UE::Geometry::FMeshMapBaker>> Compute = nullptr;
 
-public:
-	// Bake settings
-	struct FUVSnapshotBakeSettings
-	{
-		UE::Geometry::FImageDimensions Dimensions;
-		int UVLayer = 0;
-		float WireframeThickness = 1.0f;
-		FLinearColor WireframeColor = FLinearColor::Blue;
-		FLinearColor ShellColor = FLinearColor::Gray;
-		FLinearColor BackgroundColor = FLinearColor(0.0f, 0.0f, 0.0f, 0.0f);
-		int SamplesPerPixel = 4;
-		bool operator==(const FUVSnapshotBakeSettings& Other) const
-		{
-			return Dimensions == Other.Dimensions &&
-				UVLayer == Other.UVLayer &&
-				WireframeThickness == Other.WireframeThickness &&
-				WireframeColor == Other.WireframeColor &&
-				ShellColor == Other.ShellColor &&
-				BackgroundColor == Other.BackgroundColor &&
-				SamplesPerPixel == Other.SamplesPerPixel;
-		}
-	};
-protected:
-	FUVSnapshotBakeSettings CachedBakeSettings;
-	
-private:
-	EUVSnapshotBakeOpState OpState = EUVSnapshotBakeOpState::Evaluate;
-
 	// used in the bake
 	TSharedPtr<UE::Geometry::FDynamicMesh3, ESPMode::ThreadSafe> DetailMesh;
 	TSharedPtr<UE::Geometry::FDynamicMeshAABBTree3, ESPMode::ThreadSafe> DetailSpatial;
 
-	/**
-	 * Process OpState flags and update background compute
-	 * Invoked during render
-	 */
-	void UpdateResult();
-
 	/**
 	 * Retrieves the result of the FMeshMapBaker and generated UTexture2D into the CachedUVMap
 	 * 
@@ -148,29 +99,12 @@ private:
 	void OnMapUpdated(const TUniquePtr<UE::Geometry::FMeshMapBaker>& NewResult);
 
 	/**
-	 * Sets the result to a nullptr
-	 */
-	void InvalidateResults() const;
-
-	/**
-	 * Invalidates the background compute operator.
-	 */
-	void InvalidateCompute();
-
-	/**
-	 * Updates the preview material on the preview mesh with the
+	 * Updates the preview material on the preview quad with the
 	 * computed results. Invoked by OnMapsUpdated.
 	 * Also sets the result to what is currently in CachedUVMap
 	 */
 	void UpdateVisualization();
 
-	/**
-	 * Updates the FUVSnapshotBakeSettings and CachedBakeSettings
-	 *
-	 * @return current result state (Clean or Evaluate if Bake settings have changed)
-	 */
-	EUVSnapshotBakeOpState UpdateResult_UVShellMap();
-
 	/**
 	 * Internal cache of bake uv texture result
 	 */
@@ -180,25 +114,18 @@ private:
 	/**
 	 * Uses Preview Geometry to draw a preview of the bake in the unwrap viewport
 	 */
-	void UpdatePreviewMaterialBasedOnBackground();
+	void SetUpPreviewQuad();
+
+	/**
+	 * Retrieves the Material for the preview quad, or the preview result
+	 */
+	UMaterialInstanceDynamic* GetMaterialForQuad();
 
 	/**
 	 * Create texture asset from our result Texture2D
 	 * @param Texture the result texture to create
-	 * @param SourceWorld the source world to define where the texture asset will be stored.
-	 * @param SourceAsset if not null, result texture will be stored adjacent to this asset.
 	 */
-	void CreateTextureAsset(const TObjectPtr<UTexture2D>& Texture, UWorld* SourceWorld, UObject* SourceAsset) const;
-};
-
-UCLASS()
-class UVEDITORTOOLS_API UUVEditorBakeUVShellResultProperties : public UInteractiveToolPropertySet
-{
-	GENERATED_BODY()
-public:
-	/** Bake */
-	UPROPERTY(VisibleAnywhere, Category = Results, meta = (DisplayName = "Result Texture", TransientToolProperty))
-	TObjectPtr<UTexture2D> Result;
+	void CreateTextureAsset(const TObjectPtr<UTexture2D>& Texture) const;
 };
 
 UCLASS()
@@ -237,4 +164,8 @@ public:
 	/** Saved path where last UVSnapshot was saved to. Empty if this is first save out */
 	UPROPERTY()
 	FString SavedPath = FString();
+
+	/** Bake */
+	UPROPERTY(VisibleAnywhere, Category = Results, meta = (DisplayName = "Result Texture", TransientToolProperty))
+	TObjectPtr<UTexture2D> Result;
 };
\ No newline at end of file
diff --git a/Engine/Plugins/Runtime/MeshModelingToolset/Source/ModelingComponents/Public/ModelingObjectsCreationAPI.h b/Engine/Plugins/Runtime/MeshModelingToolset/Source/ModelingComponents/Public/ModelingObjectsCreationAPI.h
--- a/Engine/Plugins/Runtime/MeshModelingToolset/Source/ModelingComponents/Public/ModelingObjectsCreationAPI.h
+++ b/Engine/Plugins/Runtime/MeshModelingToolset/Source/ModelingComponents/Public/ModelingObjectsCreationAPI.h
@@ -282,6 +282,12 @@ struct MODELINGCOMPONENTS_API FCreateTextureObjectParams
 	 */
 	UPROPERTY(Category = "CreateTextureObjectParams", EditAnywhere)
 	TObjectPtr<UTexture2D> GeneratedTransientTexture = nullptr;
+
+	/**
+	 * A full path location to save the asset out to. If this parameter is not null, it overrides other work done to find a path
+	 */
+	UPROPERTY(Category= "CreateTextureObjectParams", EditAnywhere)
+	FString FullAssetPath = "";
 };
 
 
diff --git a/Engine/Plugins/Runtime/MeshModelingToolset/Source/ModelingComponentsEditorOnly/Private/EditorModelingObjectsCreationAPI.cpp b/Engine/Plugins/Runtime/MeshModelingToolset/Source/ModelingComponentsEditorOnly/Private/EditorModelingObjectsCreationAPI.cpp
--- a/Engine/Plugins/Runtime/MeshModelingToolset/Source/ModelingComponentsEditorOnly/Private/EditorModelingObjectsCreationAPI.cpp
+++ b/Engine/Plugins/Runtime/MeshModelingToolset/Source/ModelingComponentsEditorOnly/Private/EditorModelingObjectsCreationAPI.cpp
@@ -414,28 +414,35 @@ FCreateMeshObjectResult UEditorModelingObjectsCreationAPI::CreateStaticMeshAsset
 
 FCreateTextureObjectResult UEditorModelingObjectsCreationAPI::CreateTextureObject(FCreateTextureObjectParams&& CreateTexParams)
 {
-	UE::AssetUtils::FTexture2DAssetOptions AssetOptions;
-
-	ECreateModelingObjectResult AssetPathResult = GetNewAssetPath(
-		AssetOptions.NewAssetPath,
-		CreateTexParams.BaseName,
-		CreateTexParams.StoreRelativeToObject,
-		CreateTexParams.TargetWorld);
-
-	if (AssetPathResult != ECreateModelingObjectResult::Ok)
-	{
-		return FCreateTextureObjectResult{ AssetPathResult };
-	}
-
 	// currently we cannot create a new texture without an existing generated texture to store
 	if (!ensure(CreateTexParams.GeneratedTransientTexture))
 	{
 		return FCreateTextureObjectResult{ ECreateModelingObjectResult::Failed_InvalidTexture };
 	}
-
+	UE::AssetUtils::FTexture2DAssetOptions AssetOptions;
 	UE::AssetUtils::FTexture2DAssetResults ResultData;
-	UE::AssetUtils::ECreateTexture2DResult AssetResult = UE::AssetUtils::SaveGeneratedTexture2DAsset(
-		CreateTexParams.GeneratedTransientTexture, AssetOptions, ResultData);
+	
+	if (CreateTexParams.FullAssetPath != "")
+	{
+		AssetOptions.NewAssetPath = CreateTexParams.FullAssetPath;
+	}
+	else
+	{
+		const ECreateModelingObjectResult AssetPathResult = GetNewAssetPath(
+			AssetOptions.NewAssetPath,
+			CreateTexParams.BaseName,
+			CreateTexParams.StoreRelativeToObject,
+			CreateTexParams.TargetWorld);
+
+		if (AssetPathResult != ECreateModelingObjectResult::Ok)
+		{
+			return FCreateTextureObjectResult{ AssetPathResult };
+		}
+
+	}
+	const UE::AssetUtils::ECreateTexture2DResult AssetResult = UE::AssetUtils::SaveGeneratedTexture2DAsset(
+			CreateTexParams.GeneratedTransientTexture, AssetOptions, ResultData);
+	
 
 	if (AssetResult != UE::AssetUtils::ECreateTexture2DResult::Ok)
 	{
