Modeling Tools: Fix bug where level editor scale snapping affected UV editor scale snapping

diff --git a/Engine/Plugins/Editor/UVEditor/Source/UVEditorTools/Private/UVSelectTool.cpp b/Engine/Plugins/Editor/UVEditor/Source/UVEditorTools/Private/UVSelectTool.cpp
--- a/Engine/Plugins/Editor/UVEditor/Source/UVEditorTools/Private/UVSelectTool.cpp
+++ b/Engine/Plugins/Editor/UVEditor/Source/UVEditorTools/Private/UVSelectTool.cpp
@@ -254,6 +254,7 @@ void UUVSelectTool::Setup()
 	TransformGizmo->bSnapToWorldGrid = false;
 	TransformGizmo->bSnapToWorldRotGrid = false;
 	TransformGizmo->bUseContextCoordinateSystem = false;
+	TransformGizmo->bSnapToScaleGrid = false;
 
 	TransformGizmo->SetActiveTarget(TransformProxy, GetToolManager());
 	
