Modeling Tools: Add validation check for Target UV Layer prop in UV Snapshot tool. Also ensure SourceTarget is valid before launching tool

diff --git a/Engine/Plugins/Editor/UVEditor/Source/UVEditorTools/Private/UVEditorUVSnapshotTool.cpp b/Engine/Plugins/Editor/UVEditor/Source/UVEditorTools/Private/UVEditorUVSnapshotTool.cpp
--- a/Engine/Plugins/Editor/UVEditor/Source/UVEditorTools/Private/UVEditorUVSnapshotTool.cpp
+++ b/Engine/Plugins/Editor/UVEditor/Source/UVEditorTools/Private/UVEditorUVSnapshotTool.cpp
@@ -24,8 +24,8 @@ using namespace UE::Geometry;
  */
 bool UUVEditorUVSnapshotToolBuilder::CanBuildTool(const FToolBuilderState& SceneState) const
 {
-	// ensure only one target is handled
-	return Targets && Targets->Num() == 1;
+	// ensure only one target is handled and that source target is valid
+	return Targets && Targets->Num() == 1 && UE::ToolTarget::GetTargetActor((*Targets)[0]->SourceTarget);
 }
 UInteractiveTool* UUVEditorUVSnapshotToolBuilder::BuildTool(const FToolBuilderState& SceneState) const
 {
@@ -85,12 +85,15 @@ void UUVEditorUVSnapshotTool::Setup()
 	UVShellSettings->RestoreProperties(this);
 	AddToolPropertySource(UVShellSettings);
 
+	const int16 NumUVLayers = Target->AppliedCanonical->Attributes()->NumUVLayers();
+	InitializeUVLayerNames(UVShellSettings->TargetUVLayerNamesList, NumUVLayers);
+
 	// retrieve whatever UV Layer is currently being displayed in UV Editor
-	UVShellSettings->UVLayer = Target->UVLayerIndex;
-	
+	UVShellSettings->UVLayer = UVShellSettings->TargetUVLayerNamesList[Target->UVLayerIndex];
+	UVShellSettings->WatchProperty(UVShellSettings->UVLayer, [this](FString) { Compute->InvalidateResult(); });
+
 	UVShellSettings->WatchProperty(UVShellSettings->SamplesPerPixel, [this](EBakeTextureSamplesPerPixel) { Compute->InvalidateResult(); });
 	UVShellSettings->WatchProperty(UVShellSettings->Resolution, [this](EBakeTextureResolution) { Compute->InvalidateResult(); });
-	UVShellSettings->WatchProperty(UVShellSettings->UVLayer, [this](int) { Compute->InvalidateResult(); });
 	UVShellSettings->WatchProperty(UVShellSettings->WireframeThickness, [this](float) { Compute->InvalidateResult(); });
 	UVShellSettings->WatchProperty(UVShellSettings->WireframeColor, [this](FLinearColor) { Compute->InvalidateResult(); });
 	UVShellSettings->WatchProperty(UVShellSettings->ShellColor, [this](FLinearColor) { Compute->InvalidateResult(); });
@@ -173,7 +176,7 @@ TUniquePtr<UE::Geometry::TGenericDataOperator<UE::Geometry::FMeshMapBaker>> UUVE
 	Op->BakerDimensions = ImgDimensions;
 	Op->BakerSamplesPerPixel = (int32)UVShellSettings->SamplesPerPixel;
 
-	Op->UVShellEval->UVLayer = UVShellSettings->UVLayer;
+	Op->UVShellEval->UVLayer = UVShellSettings->TargetUVLayerNamesList.IndexOfByKey(UVShellSettings->UVLayer);
 	Op->UVShellEval->WireframeThickness = UVShellSettings->WireframeThickness;
 	Op->UVShellEval->WireframeColor = UVShellSettings->WireframeColor;
 	Op->UVShellEval->ShellColor = UVShellSettings->ShellColor;
@@ -266,7 +269,9 @@ void UUVEditorUVSnapshotTool::CreateTextureAsset(const TObjectPtr<UTexture2D>& T
 {
 	bool bCreatedAssetOK = true;
 	const FString ObjName = UE::ToolTarget::GetTargetActor(Target->SourceTarget)->GetActorNameOrLabel();
-	FString NewAssetName = FString::Printf(TEXT("%s_UVShell_UV%d"), *ObjName, UVShellSettings->UVLayer); // will be something like "Cylinder_UVShell_UV0"
+	FString UVLayerAsString = UVShellSettings->UVLayer;
+	UVLayerAsString.RemoveSpacesInline();
+	FString NewAssetName = FString::Printf(TEXT("%s_UVShell_%s"), *ObjName, *UVLayerAsString); // will be something like "Cylinder_UVShell_UV0"
 
 	// open dialog so user can choose where to save out the new asset
 	IContentBrowserSingleton& ContentBrowser = FModuleManager::LoadModuleChecked<FContentBrowserModule>("ContentBrowser").Get();
@@ -301,4 +306,14 @@ void UUVEditorUVSnapshotTool::CreateTextureAsset(const TObjectPtr<UTexture2D>& T
 		
 	ensure(bCreatedAssetOK);
 }
-#undef LOCTEXT_NAMESPACE
\ No newline at end of file
+
+void UUVEditorUVSnapshotTool::InitializeUVLayerNames(TArray<FString>& UVLayerNamesList, const int16 NumUVLayers)
+{
+	UVLayerNamesList.Reset();
+	for (int16 k = 0; k < NumUVLayers; ++k)
+	{
+		UVLayerNamesList.Add(FString::Printf(TEXT("UV %d"), k));
+	}
+}
+
+#undef LOCTEXT_NAMESPACE
diff --git a/Engine/Plugins/Editor/UVEditor/Source/UVEditorTools/Public/UVEditorUVSnapshotTool.h b/Engine/Plugins/Editor/UVEditor/Source/UVEditorTools/Public/UVEditorUVSnapshotTool.h
--- a/Engine/Plugins/Editor/UVEditor/Source/UVEditorTools/Public/UVEditorUVSnapshotTool.h
+++ b/Engine/Plugins/Editor/UVEditor/Source/UVEditorTools/Public/UVEditorUVSnapshotTool.h
@@ -126,6 +126,11 @@ private:
 	 * @param Texture the result texture to create
 	 */
 	void CreateTextureAsset(const TObjectPtr<UTexture2D>& Texture) const;
+
+	/**
+	 * Initialize the list of all UV Layer names
+	 */
+	static void InitializeUVLayerNames(TArray<FString>& UVLayerNamesList, const int16 NumUVLayers);
 };
 
 UCLASS()
@@ -134,8 +139,8 @@ class UVEDITORTOOLS_API UUVEditorBakeUVShellProperties : public UInteractiveTool
 	GENERATED_BODY()
 public:
 	/** The source mesh UV Layer to sample. */
-	UPROPERTY(EditAnywhere, Category = "UV Snapshot Output", meta = (DisplayName = "UV Layer"))
-	int UVLayer = 0;
+	UPROPERTY(EditAnywhere, Category = "UV Snapshot Output", meta = (DisplayName = "UV Layer", GetOptions = GetTargetUVLayerNamesFunc))
+	FString UVLayer;
 
 	/** The thickness of the wireframe in pixels. */
 	UPROPERTY(EditAnywhere, Category = "UV Snapshot Output", meta = (UIMin = "0.0", UIMax = "10.0", ClampMin = "0.0"))
@@ -168,4 +173,13 @@ public:
 	/** Bake */
 	UPROPERTY(VisibleAnywhere, Category = Results, meta = (DisplayName = "Result Texture", TransientToolProperty))
 	TObjectPtr<UTexture2D> Result;
+
+	UFUNCTION()
+	const TArray<FString>& GetTargetUVLayerNamesFunc() const
+	{
+		return TargetUVLayerNamesList;
+	}
+
+	UPROPERTY(meta = (TransientToolProperty))
+	TArray<FString> TargetUVLayerNamesList;
 };
\ No newline at end of file
