UVEditor: Fixed crash by preventing UV Layer property from being reset to default.
Fixed Snapshot tool being disabled when launching the UV Editor from the Modeling Tools palette

For the crash with the UV Layer property, we fix the crash by preventing ResetToDefault on the property. However, we also hide this property to bring it in line with the way other tools operate in the UV Editor. Users are expected to change the UV Channel from the UV Editor toolbar which will cancel tools that are already in flight on a different UV channel.

diff --git a/Engine/Plugins/Editor/UVEditor/Source/UVEditorTools/Private/UVEditorUVSnapshotTool.cpp b/Engine/Plugins/Editor/UVEditor/Source/UVEditorTools/Private/UVEditorUVSnapshotTool.cpp
--- a/Engine/Plugins/Editor/UVEditor/Source/UVEditorTools/Private/UVEditorUVSnapshotTool.cpp
+++ b/Engine/Plugins/Editor/UVEditor/Source/UVEditorTools/Private/UVEditorUVSnapshotTool.cpp
@@ -25,7 +25,7 @@ using namespace UE::Geometry;
 bool UUVEditorUVSnapshotToolBuilder::CanBuildTool(const FToolBuilderState& SceneState) const
 {
 	// ensure only one target is handled and that source target is valid
-	return Targets && Targets->Num() == 1 && UE::ToolTarget::GetTargetActor((*Targets)[0]->SourceTarget);
+	return Targets && Targets->Num() == 1;
 }
 UInteractiveTool* UUVEditorUVSnapshotToolBuilder::BuildTool(const FToolBuilderState& SceneState) const
 {
@@ -268,7 +268,7 @@ UMaterialInstanceDynamic* UUVEditorUVSnapshotTool::GetMaterialForQuad()
 void UUVEditorUVSnapshotTool::CreateTextureAsset(const TObjectPtr<UTexture2D>& Texture) const
 {
 	bool bCreatedAssetOK = true;
-	const FString ObjName = UE::ToolTarget::GetTargetActor(Target->SourceTarget)->GetActorNameOrLabel();
+	FString ObjName = UE::ToolTarget::GetHumanReadableName(Target->SourceTarget);
 	FString UVLayerAsString = UVShellSettings->UVLayer;
 	UVLayerAsString.RemoveSpacesInline();
 	FString NewAssetName = FString::Printf(TEXT("%s_UVShell_%s"), *ObjName, *UVLayerAsString); // will be something like "Cylinder_UVShell_UV0"
diff --git a/Engine/Plugins/Editor/UVEditor/Source/UVEditorTools/Public/UVEditorUVSnapshotTool.h b/Engine/Plugins/Editor/UVEditor/Source/UVEditorTools/Public/UVEditorUVSnapshotTool.h
--- a/Engine/Plugins/Editor/UVEditor/Source/UVEditorTools/Public/UVEditorUVSnapshotTool.h
+++ b/Engine/Plugins/Editor/UVEditor/Source/UVEditorTools/Public/UVEditorUVSnapshotTool.h
@@ -139,7 +139,7 @@ class UVEDITORTOOLS_API UUVEditorBakeUVShellProperties : public UInteractiveTool
 	GENERATED_BODY()
 public:
 	/** The source mesh UV Layer to sample. */
-	UPROPERTY(EditAnywhere, Category = "UV Snapshot Output", meta = (DisplayName = "UV Layer", GetOptions = GetTargetUVLayerNamesFunc))
+	UPROPERTY(meta = (DisplayName = "UV Layer", GetOptions = GetTargetUVLayerNamesFunc, NoResetToDefault))
 	FString UVLayer;
 
 	/** The thickness of the wireframe in pixels. */
