Modeling Mode: ensure UV Snapshot tool cannot be accepted until after computation is complete and add user notification while in the middle of computing

diff --git a/Engine/Plugins/Editor/UVEditor/Source/UVEditorTools/Private/UVEditorUVSnapshotTool.cpp b/Engine/Plugins/Editor/UVEditor/Source/UVEditorTools/Private/UVEditorUVSnapshotTool.cpp
--- a/Engine/Plugins/Editor/UVEditor/Source/UVEditorTools/Private/UVEditorUVSnapshotTool.cpp
+++ b/Engine/Plugins/Editor/UVEditor/Source/UVEditorTools/Private/UVEditorUVSnapshotTool.cpp
@@ -8,6 +8,7 @@
 #include "AssetUtils/Texture2DBuilder.h"
 #include "Sampling/MeshUVShellMapEvaluator.h"
 #include "ModelingObjectsCreationAPI.h"
+#include "ToolSetupUtil.h"
 #include "Materials/Material.h"
 #include "Materials/MaterialInstanceDynamic.h"
 #include "UVEditorUXSettings.h"
@@ -156,6 +157,9 @@ void UUVEditorUVSnapshotTool::Shutdown(EToolShutdownType ShutdownType)
 	// re-enable wireframe display and unwrap preview
 	Target->WireframeDisplay->Settings->bVisible = true;
 	Target->UnwrapPreview->SetVisibility(true);
+
+	// remove 'in progress' material from 3D/live preview viewport
+	Target->AppliedPreview->OverrideMaterial = nullptr;
 }
 void UUVEditorUVSnapshotTool::OnTick(float DeltaTime)
 {
@@ -163,6 +167,32 @@ void UUVEditorUVSnapshotTool::OnTick(float DeltaTime)
 	{
 		Compute->Tick(DeltaTime);
 	}
+
+	// inform the user if computation is taking a longer time
+	FText Message = FText();
+
+	if (Compute->HaveValidResult())
+	{
+		PreviewGeoBackgroundQuad->SetAllVisible(true);
+	}
+	else
+	{
+		Message = FText(LOCTEXT("Computing", "Computing..."));
+
+		// hides unwrap/2d preview when computation is occurring
+		PreviewGeoBackgroundQuad->SetAllVisible(false);
+		// applies scrolling in progress material to 3d view while computation is occurring
+		Target->AppliedPreview->OverrideMaterial = ToolSetupUtil::GetDefaultWorkingMaterial(GetToolManager());
+	}
+
+	// while computing, message will be displayed to user that it's in progress, otherwise nothing
+	GetToolManager()->DisplayMessage(Message, EToolMessageLevel::UserWarning);
+}
+
+bool UUVEditorUVSnapshotTool::CanAccept() const
+{
+	// while computing is still occurring, do not allow accept of tool
+	return Compute->HaveValidResult();
 }
 
 TUniquePtr<UE::Geometry::TGenericDataOperator<UE::Geometry::FMeshMapBaker>> UUVEditorUVSnapshotTool::MakeNewOperator()
@@ -213,6 +243,11 @@ void UUVEditorUVSnapshotTool::UpdateVisualization()
 	UTriangleSetComponent* TriangleSet = PreviewGeoBackgroundQuad->FindTriangleSet("UVShellMap");
 	TriangleSet->SetAllTrianglesMaterial(GetMaterialForQuad());
 
+	// apply UV map preview to 3D view
+	UMaterialInstanceDynamic* Result3DMaterial = UMaterialInstanceDynamic::Create(LoadObject<UMaterial>(nullptr, TEXT("/UVEditor/Materials/UVEditorBackground")), this);
+	Result3DMaterial->SetTextureParameterValue(TEXT("BackgroundBaseMap_Color"), UVShellSettings->Result);
+	Target->AppliedPreview->OverrideMaterial = Result3DMaterial;
+
 	// setting visibility here and not in setup to avoid brief moment between tool activation and initial result computation & display
 	// which showed default 'BackgroundBaseMap_Color' texture
 	PreviewGeoBackgroundQuad->SetAllVisible(true);
diff --git a/Engine/Plugins/Editor/UVEditor/Source/UVEditorTools/Public/UVEditorUVSnapshotTool.h b/Engine/Plugins/Editor/UVEditor/Source/UVEditorTools/Public/UVEditorUVSnapshotTool.h
--- a/Engine/Plugins/Editor/UVEditor/Source/UVEditorTools/Public/UVEditorUVSnapshotTool.h
+++ b/Engine/Plugins/Editor/UVEditor/Source/UVEditorTools/Public/UVEditorUVSnapshotTool.h
@@ -49,7 +49,7 @@ public:
 	virtual void Render(IToolsContextRenderAPI* RenderAPI) override { };
 	virtual bool HasCancel() const override { return true; }
 	virtual bool HasAccept() const override { return true; }
-	virtual bool CanAccept() const override { return true; }
+	virtual bool CanAccept() const override;
 	// End UInteractiveTool interface
 
 	// Begin IGenericDataOperatorFactory interface
