UVEditor: Fixed issue where notifications/warnings emitted from Actions would not post.

diff --git a/Engine/Plugins/Editor/UVEditor/Source/UVEditor/Private/UVEditorMode.cpp b/Engine/Plugins/Editor/UVEditor/Source/UVEditor/Private/UVEditorMode.cpp
--- a/Engine/Plugins/Editor/UVEditor/Source/UVEditor/Private/UVEditorMode.cpp
+++ b/Engine/Plugins/Editor/UVEditor/Source/UVEditor/Private/UVEditorMode.cpp
@@ -653,7 +653,10 @@ void UUVEditorMode::RegisterActions()
 			FExecuteAction::CreateWeakLambda(Action, [this, Action]() 
 			{
 				// If we're activating from some other selection tool, go ahead and switch to the regular one
-				ActivateDefaultTool();
+				if (!IsDefaultToolActive())
+				{
+					ActivateDefaultTool();
+				}
 
 				Action->ExecuteAction();
 				for (const FUVToolSelection& Selection : SelectionAPI->GetSelections())
