UVEditor: Fix the localization support for the Transform tool's Advanced Transform category name.

diff --git a/Engine/Plugins/Editor/UVEditor/Source/UVEditorToolsEditorOnly/Private/DetailsCustomizations/UVTransformToolCustomizations.cpp b/Engine/Plugins/Editor/UVEditor/Source/UVEditorToolsEditorOnly/Private/DetailsCustomizations/UVTransformToolCustomizations.cpp
--- a/Engine/Plugins/Editor/UVEditor/Source/UVEditorToolsEditorOnly/Private/DetailsCustomizations/UVTransformToolCustomizations.cpp
+++ b/Engine/Plugins/Editor/UVEditor/Source/UVEditorToolsEditorOnly/Private/DetailsCustomizations/UVTransformToolCustomizations.cpp
@@ -115,6 +115,7 @@ void FUVEditorUVTransformToolDetails::CustomizeDetails(IDetailLayoutBuilder& Det
 
 	TArray<TWeakObjectPtr<UObject>> ObjectsBeingCustomized;
 	DetailBuilder.GetObjectsBeingCustomized(ObjectsBeingCustomized);
+	DetailBuilder.EditCategory(TEXT("Advanced Transform"), LOCTEXT("AdvancedTransformCategoryName", "Advanced Transform"));
 	check(ObjectsBeingCustomized.Num() > 0);
 	BuildQuickTranslateMenu(DetailBuilder);
 	BuildQuickRotateMenu(DetailBuilder);
