Modeling Mode: clean up code related to UV Editor 2D Toolbar

diff --git a/Engine/Plugins/Editor/UVEditor/Source/UVEditor/Private/SUVEditor2DViewportToolBar.cpp b/Engine/Plugins/Editor/UVEditor/Source/UVEditor/Private/SUVEditor2DViewportToolBar.cpp
--- a/Engine/Plugins/Editor/UVEditor/Source/UVEditor/Private/SUVEditor2DViewportToolBar.cpp
+++ b/Engine/Plugins/Editor/UVEditor/Source/UVEditor/Private/SUVEditor2DViewportToolBar.cpp
@@ -189,7 +189,6 @@ TSharedRef<SWidget> SUVEditor2DViewportToolBar::MakeTransformToolBar(const TShar
 			// explicitly specify what this widget should look like as a menu item
 			FNewMenuDelegate::CreateLambda([this, Command](FMenuBuilder& InMenuBuilder)
 				{
-					// TODO - debug why can't just use the Command / mapping isn't working 
 					InMenuBuilder.AddMenuEntry(Command);
 
 					InMenuBuilder.AddWrapperSubMenu(
