UVEditor: Fix ensure firing by waiting for completion of load of static mesh assets when UV editor is started on them.

diff --git a/Engine/Plugins/Editor/UVEditor/Source/UVEditor/Private/UVEditorModule.cpp b/Engine/Plugins/Editor/UVEditor/Source/UVEditor/Private/UVEditorModule.cpp
--- a/Engine/Plugins/Editor/UVEditor/Source/UVEditor/Private/UVEditorModule.cpp
+++ b/Engine/Plugins/Editor/UVEditor/Source/UVEditor/Private/UVEditorModule.cpp
@@ -20,6 +20,21 @@
 
 #define LOCTEXT_NAMESPACE "FUVEditorModule"
 
+namespace UVEditorModuleLocals
+{
+	void BlockUntilAssetsReady(const TArray<TObjectPtr<UObject>>& Assets)
+	{
+		for (const TObjectPtr<UObject>& Asset : Assets)
+		{
+			if (UStaticMesh* StaticMesh = Cast<UStaticMesh>(Asset))
+			{
+				// This blocks until the lods are ready
+				StaticMesh->GetNumSourceModels();
+			}
+		}
+	}
+}
+
 void FUVEditorModule::StartupModule()
 {
 	UVEditorAssetEditor = MakeShared<UE::Geometry::FUVEditorModularFeature>();
@@ -77,6 +92,8 @@ void FUVEditorModule::ShutdownModule()
 
 void FUVEditorModule::RegisterMenus()
 {
+	using namespace UVEditorModuleLocals;
+
 	// Allows cleanup when module unloads.
 	FToolMenuOwnerScoped OwnerScoped(this);
 
@@ -105,6 +122,13 @@ void FUVEditorModule::RegisterMenus()
 								TArray<TObjectPtr<UObject>> AssetsToEdit;
 								AssetsToEdit.Append(Context->LoadSelectedObjects<UObject>());
 
+								// The checks we do above are done on asset data, and the UVSubsystem->AreObjectsValidTargets call
+								//  we do below is done on loaded objects. However, it would be possible for the former to pass while
+								//  the latter fails if the target is a static mesh that is valid, but is currently being compiled 
+								//  (see the early out in UStaticMeshReadOnlyToolTarget::HasNonGeneratedLOD).
+								// What we want to do is let the static mesh finish compiling if needed, which is what this call does.
+								BlockUntilAssetsReady(AssetsToEdit);
+
 								// If we fail the ensure here, then there must be something that we're failing to check properly
 								// in AreAssetsValidTargets that we would need to track down and check in the asset data.
 								if (ensure(UVSubsystem->AreObjectsValidTargets(AssetsToEdit)))
