UVEditor: Fixed UVEditor launch conditions to support BP actors with mesh components.
Fixed issue where the RMB context menu would not display the UV Editor option for actors with any invalid components.

diff --git a/Engine/Plugins/Editor/ModelingToolsEditorMode/Source/ModelingToolsEditorMode/Private/ModelingToolsEditorMode.cpp b/Engine/Plugins/Editor/ModelingToolsEditorMode/Source/ModelingToolsEditorMode/Private/ModelingToolsEditorMode.cpp
--- a/Engine/Plugins/Editor/ModelingToolsEditorMode/Source/ModelingToolsEditorMode/Private/ModelingToolsEditorMode.cpp
+++ b/Engine/Plugins/Editor/ModelingToolsEditorMode/Source/ModelingToolsEditorMode/Private/ModelingToolsEditorMode.cpp
@@ -1234,8 +1234,14 @@ void UModelingToolsEditorMode::RegisterUVEditor()
 					TArray<TObjectPtr<UObject>> SelectedObjects;
 					UseToolsContext->GetParentEditorModeManager()->GetSelectedActors()->GetSelectedObjects(SelectedActors);
 					UseToolsContext->GetParentEditorModeManager()->GetSelectedComponents()->GetSelectedObjects(SelectedComponents);
-					SelectedObjects.Append(SelectedActors);
-					SelectedObjects.Append(SelectedComponents);
+					if (SelectedComponents.Num() > 0)
+					{
+						SelectedObjects.Append(SelectedComponents);
+					}
+					else
+					{
+						SelectedObjects.Append(SelectedActors);
+					}
 					return UVEditorAPI->CanLaunchUVEditor(SelectedObjects);
 				}
 				return false;
diff --git a/Engine/Plugins/Editor/UVEditor/Source/UVEditor/Private/UVEditorModularFeature.cpp b/Engine/Plugins/Editor/UVEditor/Source/UVEditor/Private/UVEditorModularFeature.cpp
--- a/Engine/Plugins/Editor/UVEditor/Source/UVEditor/Private/UVEditorModularFeature.cpp
+++ b/Engine/Plugins/Editor/UVEditor/Source/UVEditor/Private/UVEditorModularFeature.cpp
@@ -31,29 +31,28 @@ namespace Geometry
 
 	void FUVEditorModularFeature::ConvertInputArgsToValidTargets(const TArray<TObjectPtr<UObject>>& ObjectsIn, TArray<TObjectPtr<UObject>>& ObjectsOut) const
 	{
+		UUVEditorSubsystem* UVSubsystem = GEditor->GetEditorSubsystem<UUVEditorSubsystem>();
+		if (!UVSubsystem)
+		{
+			return;
+		}
+		
 		for (const TObjectPtr<UObject>& Object : ObjectsIn)
 		{
-			const AActor* Actor = Cast<const AActor>(Object);
-			if (Actor)
+			if (const AActor* Actor = Cast<const AActor>(Object))
 			{
-				TArray<UObject*> ActorAssets;
-				Actor->GetReferencedContentObjects(ActorAssets);
-
-				if (ActorAssets.Num() > 0)
+				// Need to transform actors to components here because that's what the UVEditor expects to have
+				TInlineComponentArray<UActorComponent*> ActorComponents;
+				Actor->GetComponents(ActorComponents);
+				for (UActorComponent* ActorComponent : ActorComponents)
 				{
-					for (UObject* Asset : ActorAssets)
+					if (UVSubsystem->IsObjectValidTarget(ActorComponent))
 					{
-						ObjectsOut.AddUnique(Asset);
+						ObjectsOut.AddUnique(ActorComponent);
 					}
 				}
-				else {
-					// Need to transform actors to components here because that's what the UVEditor expects to have
-					TInlineComponentArray<UActorComponent*> ActorComponents;
-					Actor->GetComponents(ActorComponents);
-					ObjectsOut.Append(ActorComponents);
-				}
 			}
-			else
+			else if (UVSubsystem->IsObjectValidTarget(Object))
 			{
 				ObjectsOut.AddUnique(Object);
 			}
diff --git a/Engine/Plugins/Editor/UVEditor/Source/UVEditor/Private/UVEditorModule.cpp b/Engine/Plugins/Editor/UVEditor/Source/UVEditor/Private/UVEditorModule.cpp
--- a/Engine/Plugins/Editor/UVEditor/Source/UVEditor/Private/UVEditorModule.cpp
+++ b/Engine/Plugins/Editor/UVEditor/Source/UVEditor/Private/UVEditorModule.cpp
@@ -169,43 +169,29 @@ void FUVEditorModule::RegisterMenus()
 	auto SetupMenuEntryParameters = []
 	{
 		FMenuEntryParameters Result;
+		Result.UVSubsystem = GEditor->GetEditorSubsystem<UUVEditorSubsystem>();
+		Result.bValidTargets = false;
+		if (!Result.UVSubsystem)
+		{
+			return Result;
+		}
 
 		for (FSelectionIterator It(GEditor->GetSelectedActorIterator()); It; ++It)
 		{
-			// We are interested in the (unique) assets backing the actor, or else the actor
-			// itself if it is not asset backed (such as UDynamicMesh).
 			const AActor* Actor = static_cast<AActor*>(*It);
-
-			// We need to determine if there are components under this actor first,
-			// to ensure we generate Component based targets. Otherwise we can use
-			// the referenced content inside to generate more direct targets. This
-			// is important because otherwise we won't have access to component
-			// level data, such as transforms within the level editor, if we actually
-			// do want it for tools.
-
 			TInlineComponentArray<UActorComponent*> ActorComponents;
 			Actor->GetComponents(ActorComponents);
 
-			if (ActorComponents.Num() > 0)
+			for (UActorComponent* ActorComponent : ActorComponents)
 			{
-				Result.TargetObjects.Append(ActorComponents);
-			}
-			else
-			{
-				TArray<UObject*> ActorAssets;
-				Actor->GetReferencedContentObjects(ActorAssets);
-				for (UObject* Asset : ActorAssets)
+				if (Result.UVSubsystem->IsObjectValidTarget(ActorComponent))
 				{
-					Result.TargetObjects.AddUnique(Asset);
+					Result.TargetObjects.Add(ActorComponent);
 				}
 			}
 		}
 
-		Result.UVSubsystem = GEditor->GetEditorSubsystem<UUVEditorSubsystem>();
-		check(Result.UVSubsystem);
-
-		Result.bValidTargets = Result.UVSubsystem->AreObjectsValidTargets(Result.TargetObjects);
-
+		Result.bValidTargets = Result.TargetObjects.Num() > 0;
 		return Result;
 	};
 
diff --git a/Engine/Plugins/Editor/UVEditor/Source/UVEditor/Private/UVEditorSubsystem.cpp b/Engine/Plugins/Editor/UVEditor/Source/UVEditor/Private/UVEditorSubsystem.cpp
--- a/Engine/Plugins/Editor/UVEditor/Source/UVEditor/Private/UVEditorSubsystem.cpp
+++ b/Engine/Plugins/Editor/UVEditor/Source/UVEditor/Private/UVEditorSubsystem.cpp
@@ -137,6 +137,11 @@ void UUVEditorSubsystem::Deinitialize()
 	ToolTargetManager = nullptr;
 }
 
+bool UUVEditorSubsystem::IsObjectValidTarget(const TObjectPtr<UObject> InObject) const
+{
+	return InObject && ToolTargetManager->CanBuildTarget(InObject, UUVEditorMode::GetToolTargetRequirements());
+}
+
 bool UUVEditorSubsystem::AreObjectsValidTargets(const TArray<UObject*>& InObjects) const
 {
 	if (InObjects.IsEmpty())
@@ -146,7 +151,7 @@ bool UUVEditorSubsystem::AreObjectsValidTargets(const TArray<UObject*>& InObject
 
 	for (UObject* Object : InObjects)
 	{
-		if (!ToolTargetManager->CanBuildTarget(Object, UUVEditorMode::GetToolTargetRequirements()))
+		if (!IsObjectValidTarget(Object))
 		{
 			return false;
 		}
@@ -155,13 +160,8 @@ bool UUVEditorSubsystem::AreObjectsValidTargets(const TArray<UObject*>& InObject
 	return true;
 }
 
-bool UUVEditorSubsystem::AreAssetsValidTargets(const TArray<FAssetData>& InAssets) const
+bool UUVEditorSubsystem::IsAssetValidTarget(const FAssetData& InAsset) const
 {
-	if (InAssets.IsEmpty())
-	{
-		return false;
-	}
-
 	// Currently our tool target factories don't evaluate FAssetData to figure out whether they can
 	// build a tool target (they only work on UObjects directly), so for now we do corresponding checks
 	// here ourselves.
@@ -178,9 +178,19 @@ bool UUVEditorSubsystem::AreAssetsValidTargets(const TArray<FAssetData>& InAsset
 		return AssetData.GetClass() == USkeletalMesh::StaticClass();
 	};
 
+	return IsValidStaticMeshAsset(InAsset) || IsValidSkeletalMeshAsset(InAsset);
+}
+
+bool UUVEditorSubsystem::AreAssetsValidTargets(const TArray<FAssetData>& InAssets) const
+{
+	if (InAssets.IsEmpty())
+	{
+		return false;
+	}
+
 	for (const FAssetData& AssetData : InAssets)
 	{
-		if (!IsValidStaticMeshAsset(AssetData) && !IsValidSkeletalMeshAsset(AssetData))
+		if (!IsAssetValidTarget(AssetData))
 		{
 			return false;
 		}
@@ -194,12 +204,26 @@ void UUVEditorSubsystem::BuildTargets(const TArray<TObjectPtr<UObject>>& Objects
 {
 	TargetsOut.Reset();
 
+	TSet<UObject*> TargetAssets;
+
 	for (UObject* Object : ObjectsIn)
 	{
-		UToolTarget* Target = ToolTargetManager->BuildTarget(Object, TargetRequirements);
-		if (Target)
+		if (UToolTarget* Target = ToolTargetManager->BuildTarget(Object, TargetRequirements))
 		{
-			TargetsOut.Add(Target);
+			// If this is an asset backed target, only add it if the asset is not already added to the list.
+			if (IAssetBackedTarget* AssetTarget = Cast<IAssetBackedTarget>(Target))
+			{
+				UObject* SourceData = AssetTarget->GetSourceData();
+				if (!TargetAssets.Contains(SourceData))
+				{
+					TargetAssets.Add(SourceData);
+					TargetsOut.Add(Target);
+				}
+			}
+			else
+			{
+				TargetsOut.Add(Target);
+			}
 		}
 	}
 }
diff --git a/Engine/Plugins/Editor/UVEditor/Source/UVEditor/Public/UVEditorSubsystem.h b/Engine/Plugins/Editor/UVEditor/Source/UVEditor/Public/UVEditorSubsystem.h
--- a/Engine/Plugins/Editor/UVEditor/Source/UVEditor/Public/UVEditorSubsystem.h
+++ b/Engine/Plugins/Editor/UVEditor/Source/UVEditor/Public/UVEditorSubsystem.h
@@ -34,9 +34,15 @@ public:
 	virtual void Initialize(FSubsystemCollectionBase& Collection) override;
 	virtual void Deinitialize() override;
 
+	/** Checks that the object is a valid target for a UV editor session. */
+	bool IsObjectValidTarget(const TObjectPtr<UObject> InObject) const;
+
 	/** Checks that all of the objects are valid targets for a UV editor session. */
 	virtual bool AreObjectsValidTargets(const TArray<UObject*>& InObjects) const;
 
+	/** Checks that the asset is a valid target for a UV Editor session. */
+	bool IsAssetValidTarget(const FAssetData& InAsset) const;
+
 	/**
 	 * Checks that all of the assets are valid targets for an editor session. This
 	 * is preferable over AreObjectsValidTargets when we have FAssetData because it
