Updated headers using UnrealCodeFixup to make sure dllstorage is on methods/staticvars instead of types. Part 2/n

diff --git a/Engine/Plugins/Editor/UVEditor/Source/UVEditor/Public/UVEditorDistortionVisualization.h b/Engine/Plugins/Editor/UVEditor/Source/UVEditor/Public/UVEditorDistortionVisualization.h
--- a/Engine/Plugins/Editor/UVEditor/Source/UVEditor/Public/UVEditorDistortionVisualization.h
+++ b/Engine/Plugins/Editor/UVEditor/Source/UVEditor/Public/UVEditorDistortionVisualization.h
@@ -11,6 +11,8 @@
 
 #include "UVEditorDistortionVisualization.generated.h"
 
+#define UE_API UVEDITOR_API
+
 UENUM()
 enum class EDistortionMetric : uint8
 {
@@ -27,8 +29,8 @@ enum class EDistortionMetric : uint8
 /**
  * Visualization settings for the UUVEditorDistortionVisualization
  */
-UCLASS()
-class UVEDITOR_API UUVEditorDistortionVisualizationProperties : public UInteractiveToolPropertySet
+UCLASS(MinimalAPI)
+class UUVEditorDistortionVisualizationProperties : public UInteractiveToolPropertySet
 {
 	GENERATED_BODY()
 public:
@@ -71,8 +73,8 @@ public:
   This class contains the logic for handling distortion visualization within the UV Editor. It manages the setting
   of triangle color and computing the potential metrics used to describe the amount of distortion present in the UVs.
  */
-UCLASS(Transient)
-class UVEDITOR_API UUVEditorDistortionVisualization : public UObject
+UCLASS(MinimalAPI, Transient)
+class UUVEditorDistortionVisualization : public UObject
 {
 	GENERATED_BODY()
 
@@ -81,17 +83,17 @@ public:
 	/**
 	 * Client must call this after construction and parameter configuration to complete setup.
 	 */
-	void Initialize();
+	UE_API void Initialize();
 
 	/**
 	 * Client must call this before destruction to clean up state.
 	 */
-	void Shutdown();
+	UE_API void Shutdown();
 
 	/**
 	 * Client must call this every frame for changes to .Settings to be reflected in rendered result.
 	 */
-	void OnTick(float DeltaTime);
+	UE_API void OnTick(float DeltaTime);
 
 
 	
@@ -109,10 +111,11 @@ protected:
 	TArray<double> PerTargetAverageTexelDensity;
 	int32 MaximumTileTextureResolution;
 
-	void UpdateVisibility();
-	void ConfigureMeshColorsForTarget(int32 TargetIndex);
-	FColor GetDistortionColorForTriangle(const UE::Geometry::FDynamicMesh3& Mesh, int UVChannel, int Tid, double TargetAverageTexelDensity);
-	void ComputeInitialMeshSurfaceAreas();
+	UE_API void UpdateVisibility();
+	UE_API void ConfigureMeshColorsForTarget(int32 TargetIndex);
+	UE_API FColor GetDistortionColorForTriangle(const UE::Geometry::FDynamicMesh3& Mesh, int UVChannel, int Tid, double TargetAverageTexelDensity);
+	UE_API void ComputeInitialMeshSurfaceAreas();
 
 };
 
+#undef UE_API
diff --git a/Engine/Plugins/Editor/UVEditor/Source/UVEditor/Public/UVEditorInitializationContext.h b/Engine/Plugins/Editor/UVEditor/Source/UVEditor/Public/UVEditorInitializationContext.h
--- a/Engine/Plugins/Editor/UVEditor/Source/UVEditor/Public/UVEditorInitializationContext.h
+++ b/Engine/Plugins/Editor/UVEditor/Source/UVEditor/Public/UVEditorInitializationContext.h
@@ -19,11 +19,11 @@ class UEditorInteractiveToolsContext;
  * Data/api's that are meant to be used by the tools would usually be placed into
  * other, non-editor-only context objects (UUVToolLivePreviewAPI, etc).
  */
-UCLASS()
-class UVEDITOR_API UUVEditorInitializationContext : public UUVToolContextObject
+UCLASS(MinimalAPI)
+class UUVEditorInitializationContext : public UUVToolContextObject
 {
 	GENERATED_BODY()
 public:
 	TWeakObjectPtr<UEditorInteractiveToolsContext> LivePreviewITC;
 	TWeakPtr<FUICommandList> LivePreviewToolkitCommands;
-};
\ No newline at end of file
+};
diff --git a/Engine/Plugins/Editor/UVEditor/Source/UVEditor/Public/UVEditorMode.h b/Engine/Plugins/Editor/UVEditor/Source/UVEditor/Public/UVEditorMode.h
--- a/Engine/Plugins/Editor/UVEditor/Source/UVEditor/Public/UVEditorMode.h
+++ b/Engine/Plugins/Editor/UVEditor/Source/UVEditor/Public/UVEditorMode.h
@@ -12,6 +12,8 @@
 
 #include "UVEditorMode.generated.h"
 
+#define UE_API UVEDITOR_API
+
 PREDECLARE_GEOMETRY(class FDynamicMesh3);
 
 class FAssetEditorModeManager;
@@ -41,8 +43,8 @@ class UUVEditorLivePreviewUXProperties;
 /**
  * Visualization settings for the UUVEditorMode's Grid
  */
-UCLASS()
-class UVEDITOR_API UUVEditorGridProperties : public UInteractiveToolPropertySet
+UCLASS(MinimalAPI)
+class UUVEditorGridProperties : public UInteractiveToolPropertySet
 {
 	GENERATED_BODY()
 public:
@@ -97,22 +99,22 @@ enum class EUVEditorModeActions
 /**
  * Settings for UDIMs in the UVEditor
  */
-UCLASS()
-class UVEDITOR_API UUVEditorUDIMProperties : public UInteractiveToolPropertySet
+UCLASS(MinimalAPI)
+class UUVEditorUDIMProperties : public UInteractiveToolPropertySet
 {
 	GENERATED_BODY()
 public:
 	void Initialize(UUVEditorMode* ParentModeIn) { ParentMode = ParentModeIn; }
-	void PostAction(EUVEditorModeActions Action);
+	UE_API void PostAction(EUVEditorModeActions Action);
 
 	UPROPERTY(EditAnywhere, Transient, Category = "UDIM", meta = (DisplayName = "UDIM Source Mesh", GetOptions = GetAssetNames))
 	FString UDIMSourceAsset;
 
 	UFUNCTION()
-	const TArray<FString>& GetAssetNames();
+	UE_API const TArray<FString>& GetAssetNames();
 
 	UFUNCTION()
-	int32 AssetByIndex() const;
+	UE_API int32 AssetByIndex() const;
 
 	/** Set UDIM Layout from selected asset's UVs */
 	UFUNCTION(CallInEditor, Category = UDIM)
@@ -131,11 +133,11 @@ public:
 	TArray<FUDIMSpecifier> ActiveUDIMs;
 
 public:
-	void InitializeAssets(const TArray<TObjectPtr<UUVEditorToolMeshInput>>& TargetsIn);
+	UE_API void InitializeAssets(const TArray<TObjectPtr<UUVEditorToolMeshInput>>& TargetsIn);
 
 private:
-	void UpdateActiveUDIMsFromTexture();
-	void UpdateActiveUDIMsFromAsset();
+	UE_API void UpdateActiveUDIMsFromTexture();
+	UE_API void UpdateActiveUDIMsFromAsset();
 
 	TWeakObjectPtr<UUVEditorMode> ParentMode;
 	TArray<FString> UVAssetNames;
@@ -416,3 +418,4 @@ private:
 	TArray<FString> ToolsThatAllowActions;
 };
 
+#undef UE_API
diff --git a/Engine/Plugins/Editor/UVEditor/Source/UVEditor/Public/UVEditorModeChannelProperties.h b/Engine/Plugins/Editor/UVEditor/Source/UVEditor/Public/UVEditorModeChannelProperties.h
--- a/Engine/Plugins/Editor/UVEditor/Source/UVEditor/Public/UVEditorModeChannelProperties.h
+++ b/Engine/Plugins/Editor/UVEditor/Source/UVEditor/Public/UVEditorModeChannelProperties.h
@@ -9,6 +9,8 @@
 
 #include "UVEditorModeChannelProperties.generated.h"
 
+#define UE_API UVEDITOR_API
+
 PREDECLARE_GEOMETRY(class FDynamicMesh3);
 
 class UInteractiveToolPropertySet;
@@ -18,8 +20,8 @@ class UUVEditorToolMeshInput;
 /**
  * UV Layer Settings for the UVEditorMode
  */
-UCLASS()
-class UVEDITOR_API UUVEditorUVChannelProperties : public UInteractiveToolPropertySet
+UCLASS(MinimalAPI)
+class UUVEditorUVChannelProperties : public UInteractiveToolPropertySet
 {
 	GENERATED_BODY()
 public:
@@ -28,13 +30,13 @@ public:
 	FString Asset;
 
 	UFUNCTION()
-	const TArray<FString>& GetAssetNames();
+	UE_API const TArray<FString>& GetAssetNames();
 
 	UPROPERTY(EditAnywhere, Category = UVChannel, meta = (DisplayName = "UV Channel", GetOptions = GetUVChannelNames))
 	FString UVChannel;
 
 	UFUNCTION()
-	const TArray<FString>& GetUVChannelNames();
+	UE_API const TArray<FString>& GetUVChannelNames();
 
 	TArray<FString> UVChannelNames;
 	TArray<FString> UVAssetNames;
@@ -44,7 +46,7 @@ public:
 
 
 public:
-	void Initialize(
+	UE_API void Initialize(
 		TArray<TObjectPtr<UToolTarget>>& TargetObjects,
 		TArray<TSharedPtr<UE::Geometry::FDynamicMesh3>> AppliedCanonicalMeshes,
 		bool bInitializeSelection);
@@ -54,24 +56,26 @@ public:
 	 * @param bUpdateIfInvalid if selection is not valid, reset UVAsset to Asset0 or empty if no assets exist
 	 * @return true if selection in UVAsset is an entry in UVAssetNames.
 	 */
-	bool ValidateUVAssetSelection(bool bUpdateIfInvalid);
+	UE_API bool ValidateUVAssetSelection(bool bUpdateIfInvalid);
 
 	/**
 	 * Verify that the UV channel selection is valid
 	 * @param bUpdateIfInvalid if selection is not valid, UVChannel to UV0 or empty if no UV channels exist
 	 * @return true if selection in UVChannel is an entry in UVChannelNamesList.
 	 */
-	bool ValidateUVChannelSelection(bool bUpdateIfInvalid);
+	UE_API bool ValidateUVChannelSelection(bool bUpdateIfInvalid);
 
 
 	/**
 	 * @return selected UV asset ID, or -1 if invalid selection
 	 */
-	int32 GetSelectedAssetID();
+	UE_API int32 GetSelectedAssetID();
 
 	/**
 	 * @param bForceToZeroOnFailure if true, then instead of returning -1 we return 0 so calling code can fallback to default UV paths
 	 * @return selected UV channel index, or -1 if invalid selection, or 0 if invalid selection and bool bForceToZeroOnFailure = true
 	 */
-	int32 GetSelectedChannelIndex(bool bForceToZeroOnFailure = false);
+	UE_API int32 GetSelectedChannelIndex(bool bForceToZeroOnFailure = false);
 };
+
+#undef UE_API
diff --git a/Engine/Plugins/Editor/UVEditor/Source/UVEditor/Public/UVEditorModularFeature.h b/Engine/Plugins/Editor/UVEditor/Source/UVEditor/Public/UVEditorModularFeature.h
--- a/Engine/Plugins/Editor/UVEditor/Source/UVEditor/Public/UVEditorModularFeature.h
+++ b/Engine/Plugins/Editor/UVEditor/Source/UVEditor/Public/UVEditorModularFeature.h
@@ -5,6 +5,8 @@
 #include "CoreMinimal.h"
 #include "GeometryProcessingInterfaces/IUVEditorModularFeature.h"
 
+#define UE_API UVEDITOR_API
+
 namespace UE
 {
 namespace Geometry
@@ -14,17 +16,19 @@ namespace Geometry
  * Connector class that allows other plugins to look for the UV Editor to see if the plugin is present,
  * (via IModularFeatures::Get().GetModularFeature and related methods), and then launch it if so.
  */
-class UVEDITOR_API FUVEditorModularFeature : public IUVEditorModularFeature
+class FUVEditorModularFeature : public IUVEditorModularFeature
 {
 public:
-	virtual void LaunchUVEditor(const TArray<TObjectPtr<UObject>>& Objects) override;
-	virtual bool CanLaunchUVEditor(const TArray<TObjectPtr<UObject>>& Objects) override;
+	UE_API virtual void LaunchUVEditor(const TArray<TObjectPtr<UObject>>& Objects) override;
+	UE_API virtual bool CanLaunchUVEditor(const TArray<TObjectPtr<UObject>>& Objects) override;
 
 protected:
-	void ConvertInputArgsToValidTargets(const TArray<TObjectPtr<UObject>>& ObjectsIn, TArray<TObjectPtr<UObject>>& ObjectsOut) const;
+	UE_API void ConvertInputArgsToValidTargets(const TArray<TObjectPtr<UObject>>& ObjectsIn, TArray<TObjectPtr<UObject>>& ObjectsOut) const;
 
 };
 
 
 }
 }
+
+#undef UE_API
diff --git a/Engine/Plugins/Editor/UVEditor/Source/UVEditor/Public/UVEditorModule.h b/Engine/Plugins/Editor/UVEditor/Source/UVEditor/Public/UVEditorModule.h
--- a/Engine/Plugins/Editor/UVEditor/Source/UVEditor/Public/UVEditorModule.h
+++ b/Engine/Plugins/Editor/UVEditor/Source/UVEditor/Public/UVEditorModule.h
@@ -7,25 +7,27 @@
 
 #include "UVEditorModularFeature.h"
 
+#define UE_API UVEDITOR_API
+
 class FLayoutExtender;
 
 /**
  * Besides the normal module things, the module class is also responsible for hooking the 
  * UV editor into existing menus.
  */
-class UVEDITOR_API FUVEditorModule : public IModuleInterface
+class FUVEditorModule : public IModuleInterface
 {
 public:
 
 	/** IModuleInterface implementation */
-	virtual void StartupModule() override;
-	virtual void ShutdownModule() override;
+	UE_API virtual void StartupModule() override;
+	UE_API virtual void ShutdownModule() override;
 
 	DECLARE_EVENT_OneParam(FUVEditorModule, FOnRegisterLayoutExtensions, FLayoutExtender&);
 	virtual FOnRegisterLayoutExtensions& OnRegisterLayoutExtensions() { return RegisterLayoutExtensions; }
 
 protected:
-	void RegisterMenus();
+	UE_API void RegisterMenus();
 
 private:
 	FOnRegisterLayoutExtensions	RegisterLayoutExtensions;
@@ -35,3 +37,5 @@ private:
 
 	TSharedPtr<UE::Geometry::FUVEditorModularFeature> UVEditorAssetEditor;
 };
+
+#undef UE_API
diff --git a/Engine/Plugins/Editor/UVEditor/Source/UVEditor/Public/UVEditorSubsystem.h b/Engine/Plugins/Editor/UVEditor/Source/UVEditor/Public/UVEditorSubsystem.h
--- a/Engine/Plugins/Editor/UVEditor/Source/UVEditor/Public/UVEditorSubsystem.h
+++ b/Engine/Plugins/Editor/UVEditor/Source/UVEditor/Public/UVEditorSubsystem.h
@@ -11,6 +11,8 @@
 
 #include "UVEditorSubsystem.generated.h"
 
+#define UE_API UVEDITOR_API
+
 PREDECLARE_GEOMETRY(class FAssetDynamicMeshTargetComboInterface);
 
 class UToolTargetManager;
@@ -25,20 +27,20 @@ class UToolTarget;
  * instances (we can't rely on the asset editor subsystem for this because the UV editor is
  * not the primary editor for meshes), and figures out what we can launch the editor on.
  */
-UCLASS()
-class UVEDITOR_API UUVEditorSubsystem : public UEditorSubsystem
+UCLASS(MinimalAPI)
+class UUVEditorSubsystem : public UEditorSubsystem
 {
 	GENERATED_BODY()
 
 public:
-	virtual void Initialize(FSubsystemCollectionBase& Collection) override;
-	virtual void Deinitialize() override;
+	UE_API virtual void Initialize(FSubsystemCollectionBase& Collection) override;
+	UE_API virtual void Deinitialize() override;
 
 	/** Checks that the object is a valid target for a UV editor session. */
 	bool IsObjectValidTarget(const TObjectPtr<UObject> InObject) const;
 
 	/** Checks that all of the objects are valid targets for a UV editor session. */
-	virtual bool AreObjectsValidTargets(const TArray<UObject*>& InObjects) const;
+	UE_API virtual bool AreObjectsValidTargets(const TArray<UObject*>& InObjects) const;
 
 	/** Checks that the asset is a valid target for a UV Editor session. */
 	bool IsAssetValidTarget(const FAssetData& InAsset) const;
@@ -49,12 +51,12 @@ public:
 	 * allows us to avoid forcing a load of the underlying UObjects (for instance to
 	 * avoid triggering a load when right clicking an asset in the content browser).
 	 */
-	virtual bool AreAssetsValidTargets(const TArray<FAssetData>& InAssets) const;
+	UE_API virtual bool AreAssetsValidTargets(const TArray<FAssetData>& InAssets) const;
 
 	/** 
 	 * Tries to build the core targets that provide meshes for UV tools to work on.
 	 */
-	virtual void BuildTargets(const TArray<TObjectPtr<UObject>>& ObjectsIn, 
+	UE_API virtual void BuildTargets(const TArray<TObjectPtr<UObject>>& ObjectsIn, 
 		const FToolTargetTypeRequirements& TargetRequirements,
 		TArray<TObjectPtr<UToolTarget>>& TargetsOut);
 
@@ -63,13 +65,13 @@ public:
 	 * these objects, if one exists, or starts up a new instance editing all of these 
 	 * objects.
 	 */
-	virtual void StartUVEditor(TArray<TObjectPtr<UObject>> ObjectsToEdit);
+	UE_API virtual void StartUVEditor(TArray<TObjectPtr<UObject>> ObjectsToEdit);
 
 	/** 
 	 * Needs to be called when a UV editor instance is closed so that the subsystem knows
 	 * to create a new one for these objects if they are opened again.
 	 */
-	virtual void NotifyThatUVEditorClosed(TArray<TObjectPtr<UObject>> ObjectsItWasEditing);
+	UE_API virtual void NotifyThatUVEditorClosed(TArray<TObjectPtr<UObject>> ObjectsItWasEditing);
 
 protected:
 
@@ -81,4 +83,6 @@ protected:
 	TObjectPtr<UToolTargetManager> ToolTargetManager = nullptr;
 
 	TMap<TObjectPtr<UObject>, TObjectPtr<UUVEditor>> OpenedEditorInstances;
-};
\ No newline at end of file
+};
+
+#undef UE_API
diff --git a/Engine/Plugins/Editor/UVEditor/Source/UVEditor/Public/UVEditorToolkit.h b/Engine/Plugins/Editor/UVEditor/Source/UVEditor/Public/UVEditorToolkit.h
--- a/Engine/Plugins/Editor/UVEditor/Source/UVEditor/Public/UVEditorToolkit.h
+++ b/Engine/Plugins/Editor/UVEditor/Source/UVEditor/Public/UVEditorToolkit.h
@@ -6,6 +6,8 @@
 
 #include "Tools/BaseAssetToolkit.h"
 
+#define UE_API UVEDITOR_API
+
 class FAdvancedPreviewScene;
 class FUVEditorModeUILayer;
 class FUVEditorPreviewModeManager;
@@ -26,36 +28,36 @@ class UUVTool2DViewportAPI;
  * initializing the UV mode.
  * Thus, the FUVEdiotrToolkit ends up being the central place for the UV Asset editor setup.
  */
-class UVEDITOR_API FUVEditorToolkit : public FBaseAssetToolkit
+class FUVEditorToolkit : public FBaseAssetToolkit
 {
 public:
-	FUVEditorToolkit(UAssetEditor* InOwningAssetEditor);
-	virtual ~FUVEditorToolkit();
+	UE_API FUVEditorToolkit(UAssetEditor* InOwningAssetEditor);
+	UE_API virtual ~FUVEditorToolkit();
 
-	static const FName InteractiveToolsPanelTabID;
-	static const FName LivePreviewTabID;
+	static UE_API const FName InteractiveToolsPanelTabID;
+	static UE_API const FName LivePreviewTabID;
 
 	FPreviewScene* GetPreviewScene() { return UnwrapScene.Get(); }
 
 	// FBaseAssetToolkit
-	virtual void CreateWidgets() override;
+	UE_API virtual void CreateWidgets() override;
 
 	// FAssetEditorToolkit
-	virtual void AddViewportOverlayWidget(TSharedRef<SWidget> InViewportOverlayWidget, int32 ZOrder = INDEX_NONE) override;
-	virtual void RemoveViewportOverlayWidget(TSharedRef<SWidget> InViewportOverlayWidget) override;
-	virtual void CreateEditorModeManager() override;
-	virtual FText GetToolkitName() const override;
-	virtual FName GetToolkitFName() const override;
-	virtual FText GetBaseToolkitName() const override;
-	virtual FText GetToolkitToolTipText() const override;
-	virtual void RegisterTabSpawners(const TSharedRef<FTabManager>& InTabManager) override;
-	virtual bool OnRequestClose(EAssetEditorCloseReason InCloseReason) override;
-	virtual void OnClose() override;
-	virtual void SaveAsset_Execute() override;
-	virtual bool CanSaveAsset() const override;
-	virtual bool CanSaveAssetAs() const override;
-	void OnToolkitHostingStarted(const TSharedRef<IToolkit>& Toolkit) override;
-	void OnToolkitHostingFinished(const TSharedRef<IToolkit>& Toolkit) override;
+	UE_API virtual void AddViewportOverlayWidget(TSharedRef<SWidget> InViewportOverlayWidget, int32 ZOrder = INDEX_NONE) override;
+	UE_API virtual void RemoveViewportOverlayWidget(TSharedRef<SWidget> InViewportOverlayWidget) override;
+	UE_API virtual void CreateEditorModeManager() override;
+	UE_API virtual FText GetToolkitName() const override;
+	UE_API virtual FName GetToolkitFName() const override;
+	UE_API virtual FText GetBaseToolkitName() const override;
+	UE_API virtual FText GetToolkitToolTipText() const override;
+	UE_API virtual void RegisterTabSpawners(const TSharedRef<FTabManager>& InTabManager) override;
+	UE_API virtual bool OnRequestClose(EAssetEditorCloseReason InCloseReason) override;
+	UE_API virtual void OnClose() override;
+	UE_API virtual void SaveAsset_Execute() override;
+	UE_API virtual bool CanSaveAsset() const override;
+	UE_API virtual bool CanSaveAssetAs() const override;
+	UE_API void OnToolkitHostingStarted(const TSharedRef<IToolkit>& Toolkit) override;
+	UE_API void OnToolkitHostingFinished(const TSharedRef<IToolkit>& Toolkit) override;
 
 	// IAssetEditorInstance
 	// This is important because if this returns true, attempting to edit a static mesh that is
@@ -64,17 +66,17 @@ public:
 
 protected:
 
-	TSharedRef<SDockTab> SpawnTab_InteractiveToolsPanel(const FSpawnTabArgs& Args);
-	TSharedRef<SDockTab> SpawnTab_LivePreview(const FSpawnTabArgs& Args);
+	UE_API TSharedRef<SDockTab> SpawnTab_InteractiveToolsPanel(const FSpawnTabArgs& Args);
+	UE_API TSharedRef<SDockTab> SpawnTab_LivePreview(const FSpawnTabArgs& Args);
 
 	// FBaseAssetToolkit
-	virtual AssetEditorViewportFactoryFunction GetViewportDelegate() override;
-	virtual TSharedPtr<FEditorViewportClient> CreateEditorViewportClient() const override;
+	UE_API virtual AssetEditorViewportFactoryFunction GetViewportDelegate() override;
+	UE_API virtual TSharedPtr<FEditorViewportClient> CreateEditorViewportClient() const override;
 
 	// FAssetEditorToolkit
-	virtual void PostInitAssetEditor() override;
-	virtual const FSlateBrush* GetDefaultTabIcon() const override;
-	virtual FLinearColor GetDefaultTabColor() const override;
+	UE_API virtual void PostInitAssetEditor() override;
+	UE_API virtual const FSlateBrush* GetDefaultTabIcon() const override;
+	UE_API virtual FLinearColor GetDefaultTabColor() const override;
 
 	/** Scene in which the 2D unwrapped uv meshes live. */
 	TUniquePtr<FPreviewScene> UnwrapScene;
@@ -98,4 +100,6 @@ protected:
 
 	TSharedPtr<FUVEditorModeUILayer> ModeUILayer;
 	TSharedPtr<FWorkspaceItem> UVEditorMenuCategory;
-};
\ No newline at end of file
+};
+
+#undef UE_API
diff --git a/Engine/Plugins/Editor/UVEditor/Source/UVEditorTools/Public/Actions/UVMakeIslandAction.h b/Engine/Plugins/Editor/UVEditor/Source/UVEditorTools/Public/Actions/UVMakeIslandAction.h
--- a/Engine/Plugins/Editor/UVEditor/Source/UVEditorTools/Public/Actions/UVMakeIslandAction.h
+++ b/Engine/Plugins/Editor/UVEditor/Source/UVEditorTools/Public/Actions/UVMakeIslandAction.h
@@ -6,6 +6,8 @@
 
 #include "UVMakeIslandAction.generated.h"
 
+#define UE_API UVEDITORTOOLS_API
+
 /**
  * Action that takes the currently selected triangles and makes a separate UV island out of them,
  *  i.e. any interior seams are removed, and seams are added around the boundary of the selection.
@@ -15,12 +17,14 @@
  *  the created islands.
  * An unwrap operation on the island is typically a good thing to follow this action.
  */
-UCLASS()
-class UVEDITORTOOLS_API UUVMakeIslandAction : public UUVToolAction
+UCLASS(MinimalAPI)
+class UUVMakeIslandAction : public UUVToolAction
 {
 	GENERATED_BODY()
 
 public:
-	virtual bool CanExecuteAction() const override;
-	virtual bool ExecuteAction() override;
-};
\ No newline at end of file
+	UE_API virtual bool CanExecuteAction() const override;
+	UE_API virtual bool ExecuteAction() override;
+};
+
+#undef UE_API
diff --git a/Engine/Plugins/Editor/UVEditor/Source/UVEditorTools/Public/Actions/UVSeamSewAction.h b/Engine/Plugins/Editor/UVEditor/Source/UVEditorTools/Public/Actions/UVSeamSewAction.h
--- a/Engine/Plugins/Editor/UVEditor/Source/UVEditorTools/Public/Actions/UVSeamSewAction.h
+++ b/Engine/Plugins/Editor/UVEditor/Source/UVEditorTools/Public/Actions/UVSeamSewAction.h
@@ -12,23 +12,27 @@
 
 #include "UVSeamSewAction.generated.h"
 
+#define UE_API UVEDITORTOOLS_API
+
 PREDECLARE_GEOMETRY(class FUVEditorDynamicMeshSelection);
 PREDECLARE_GEOMETRY(class FDynamicMesh3);
 class APreviewGeometryActor;
 class ULineSetComponent;
 
-UCLASS()
-class UVEDITORTOOLS_API UUVSeamSewAction : public UUVToolAction
+UCLASS(MinimalAPI)
+class UUVSeamSewAction : public UUVToolAction
 {	
 	GENERATED_BODY()
 
 	using FDynamicMesh3 = UE::Geometry::FDynamicMesh3;
 
 public:
-	virtual bool CanExecuteAction() const override;
-	virtual bool ExecuteAction() override;
+	UE_API virtual bool CanExecuteAction() const override;
+	UE_API virtual bool ExecuteAction() override;
 
-	static int32 FindSewEdgeOppositePairing(const FDynamicMesh3& UnwrapMesh, 
+	static UE_API int32 FindSewEdgeOppositePairing(const FDynamicMesh3& UnwrapMesh, 
 		const FDynamicMesh3& AppliedMesh, int32 UVLayerIndex, int32 UnwrapEid, 
 		bool& bWouldPreferOppositeOrderOut);
-};
\ No newline at end of file
+};
+
+#undef UE_API
diff --git a/Engine/Plugins/Editor/UVEditor/Source/UVEditorTools/Public/Actions/UVSplitAction.h b/Engine/Plugins/Editor/UVEditor/Source/UVEditorTools/Public/Actions/UVSplitAction.h
--- a/Engine/Plugins/Editor/UVEditor/Source/UVEditorTools/Public/Actions/UVSplitAction.h
+++ b/Engine/Plugins/Editor/UVEditor/Source/UVEditorTools/Public/Actions/UVSplitAction.h
@@ -7,12 +7,16 @@
 
 #include "UVSplitAction.generated.h"
 
-UCLASS()
-class UVEDITORTOOLS_API UUVSplitAction : public UUVToolAction
+#define UE_API UVEDITORTOOLS_API
+
+UCLASS(MinimalAPI)
+class UUVSplitAction : public UUVToolAction
 {	
 	GENERATED_BODY()
 
 public:
-	virtual bool CanExecuteAction() const override;
-	virtual bool ExecuteAction() override;
-};
\ No newline at end of file
+	UE_API virtual bool CanExecuteAction() const override;
+	UE_API virtual bool ExecuteAction() override;
+};
+
+#undef UE_API
diff --git a/Engine/Plugins/Editor/UVEditor/Source/UVEditorTools/Public/Actions/UVToolAction.h b/Engine/Plugins/Editor/UVEditor/Source/UVEditorTools/Public/Actions/UVToolAction.h
--- a/Engine/Plugins/Editor/UVEditor/Source/UVEditorTools/Public/Actions/UVToolAction.h
+++ b/Engine/Plugins/Editor/UVEditor/Source/UVEditorTools/Public/Actions/UVToolAction.h
@@ -6,18 +6,20 @@
 
 #include "UVToolAction.generated.h"
 
+#define UE_API UVEDITORTOOLS_API
+
 class UUVToolEmitChangeAPI;
 class UUVToolSelectionAPI;
 class UInteractiveToolManager;
 
-UCLASS()
-class UVEDITORTOOLS_API UUVToolAction : public UObject
+UCLASS(MinimalAPI)
+class UUVToolAction : public UObject
 {
 	GENERATED_BODY()
 public:
 
-	virtual void Setup(UInteractiveToolManager* ToolManager);
-	virtual void Shutdown();
+	UE_API virtual void Setup(UInteractiveToolManager* ToolManager);
+	UE_API virtual void Shutdown();
 	virtual bool CanExecuteAction() const { return false;}
 	virtual bool ExecuteAction() { return false; }
 
@@ -32,3 +34,5 @@ protected:
 
 	virtual UInteractiveToolManager* GetToolManager() const { return ToolManager.Get(); }
 };
+
+#undef UE_API
diff --git a/Engine/Plugins/Editor/UVEditor/Source/UVEditorTools/Public/ContextObjects/UVToolAssetInputsContext.h b/Engine/Plugins/Editor/UVEditor/Source/UVEditorTools/Public/ContextObjects/UVToolAssetInputsContext.h
--- a/Engine/Plugins/Editor/UVEditor/Source/UVEditorTools/Public/ContextObjects/UVToolAssetInputsContext.h
+++ b/Engine/Plugins/Editor/UVEditor/Source/UVEditorTools/Public/ContextObjects/UVToolAssetInputsContext.h
@@ -11,8 +11,8 @@
 /**
  * Simple context object used by UV asset editor to pass inputs down to the UV editor mode.
  */
-UCLASS()
-class UVEDITORTOOLS_API UUVToolAssetInputsContext : public UUVToolContextObject
+UCLASS(MinimalAPI)
+class UUVToolAssetInputsContext : public UUVToolContextObject
 {
 	GENERATED_BODY()
 public:
@@ -33,4 +33,4 @@ public:
 protected:
 	TArray<TObjectPtr<UObject>> Assets;
 	TArray<FTransform> Transforms;
-};
\ No newline at end of file
+};
diff --git a/Engine/Plugins/Editor/UVEditor/Source/UVEditorTools/Public/ContextObjects/UVToolContextObjects.h b/Engine/Plugins/Editor/UVEditor/Source/UVEditorTools/Public/ContextObjects/UVToolContextObjects.h
--- a/Engine/Plugins/Editor/UVEditor/Source/UVEditorTools/Public/ContextObjects/UVToolContextObjects.h
+++ b/Engine/Plugins/Editor/UVEditor/Source/UVEditorTools/Public/ContextObjects/UVToolContextObjects.h
@@ -11,6 +11,8 @@
 
 #include "UVToolContextObjects.generated.h"
 
+#define UE_API UVEDITORTOOLS_API
+
 // TODO: This should be spread out across multiple files
 
 PREDECLARE_GEOMETRY(class FDynamicMesh3);
@@ -25,8 +27,8 @@ class UUVEditorToolMeshInput;
 /**
  * Base class for context objects used in the UV editor.
  */
-UCLASS()
-class UVEDITORTOOLS_API UUVToolContextObject : public UObject
+UCLASS(MinimalAPI)
+class UUVToolContextObject : public UObject
 {
 	GENERATED_BODY()
 public:
@@ -47,8 +49,8 @@ public:
  * An API object meant to be stored in a context object store that allows UV editor tools
  * to emit appropriate undo/redo transactions.
  */
-UCLASS()
-class UVEDITORTOOLS_API UUVToolEmitChangeAPI : public UUVToolContextObject
+UCLASS(MinimalAPI)
+class UUVToolEmitChangeAPI : public UUVToolContextObject
 {
 	GENERATED_BODY()
 public:
@@ -78,21 +80,21 @@ public:
 	 * preferable to use EmitToolIndependentUnwrapCanonicalChange, which will set up a proper transaction
 	 * for you.
 	 */
-	virtual void EmitToolIndependentChange(UObject* TargetObject, TUniquePtr<FToolCommandChange> Change, const FText& Description);
+	UE_API virtual void EmitToolIndependentChange(UObject* TargetObject, TUniquePtr<FToolCommandChange> Change, const FText& Description);
 
 	/**
 	 * A convenience function that is like EmitToolIndependentChange, but uses a FDynamicMeshChange
 	 * that operates on the UnwrapCanonical of an input to create a change object that updates the other
 	 * views and issues an OnUndoRedo broadcast on the input object.
 	 */
-	virtual void EmitToolIndependentUnwrapCanonicalChange(UUVEditorToolMeshInput* InputObject,
+	UE_API virtual void EmitToolIndependentUnwrapCanonicalChange(UUVEditorToolMeshInput* InputObject,
 		TUniquePtr<UE::Geometry::FDynamicMeshChange> UnwrapCanonicalMeshChange, const FText& Description);
 
 	/**
 	 * Emits a change that is considered expired when the active tool does not match the tool that was active
 	 * when it was emitted.
 	 */
-	virtual void EmitToolDependentChange(UObject* TargetObject, TUniquePtr<FToolCommandChange> Change, const FText& Description);
+	UE_API virtual void EmitToolDependentChange(UObject* TargetObject, TUniquePtr<FToolCommandChange> Change, const FText& Description);
 
 protected:
 	TWeakObjectPtr<UInteractiveToolManager> ToolManager = nullptr;
@@ -102,13 +104,13 @@ protected:
  * Allows tools to interact with the 3d preview viewport, which has a separate
  * world and input router.
  */
-UCLASS()
-class UVEDITORTOOLS_API UUVToolLivePreviewAPI : public UUVToolContextObject
+UCLASS(MinimalAPI)
+class UUVToolLivePreviewAPI : public UUVToolContextObject
 {
 	GENERATED_BODY()
 public:
 
-	void Initialize(UWorld* WorldIn, UInputRouter* RouterIn,
+	UE_API void Initialize(UWorld* WorldIn, UInputRouter* RouterIn,
 		TUniqueFunction<void(FViewCameraState& CameraStateOut)> GetLivePreviewCameraStateFuncIn,
 		TUniqueFunction<void(const UE::Geometry::FAxisAlignedBox3d& BoundingBox)> SetLivePreviewCameraToLookAtVolumeFuncIn,
 		TUniqueFunction<void(const EMouseCursor::Type Cursor, bool bEnableOverride)> SetCursorOverrideFuncIn, 
@@ -149,7 +151,7 @@ public:
 		}
 	}
 
-	virtual void OnToolEnded(UInteractiveTool* DeadTool) override;
+	UE_API virtual void OnToolEnded(UInteractiveTool* DeadTool) override;
 
 	/**
 	 * Broadcast by the 3D live preview viewport on Render() so that mechanics/tools can
@@ -181,7 +183,7 @@ private:
 };
 
 USTRUCT()
-struct UVEDITORTOOLS_API FUDIMBlock
+struct FUDIMBlock
 {
 	GENERATED_BODY();
 
@@ -191,16 +193,16 @@ struct UVEDITORTOOLS_API FUDIMBlock
 	UPROPERTY()
 	int32 TextureResolution = 512;
 
-	int32 BlockU() const;
-	int32 BlockV() const;
-	void SetFromBlocks(int32 BlockU, int32 BlockV);
+	UE_API int32 BlockU() const;
+	UE_API int32 BlockV() const;
+	UE_API void SetFromBlocks(int32 BlockU, int32 BlockV);
 };
 
 /**
  * Allows tools to interact with the 2d preview viewport 
  */
-UCLASS()
-class UVEDITORTOOLS_API UUVTool2DViewportAPI : public UUVToolContextObject
+UCLASS(MinimalAPI)
+class UUVTool2DViewportAPI : public UUVToolContextObject
 {
 	GENERATED_BODY()
 public:
@@ -273,8 +275,8 @@ protected:
 /**
  * Allows tools to interact with the assets and their UV layers
 */
-UCLASS()
-class UVEDITORTOOLS_API UUVToolAssetAndChannelAPI : public UUVToolContextObject
+UCLASS(MinimalAPI)
+class UUVToolAssetAndChannelAPI : public UUVToolContextObject
 {
 	GENERATED_BODY()
 public:
@@ -316,8 +318,8 @@ public:
  * Binds to the input objects' OnCanonicalModified delegate to automatically update 
  * the tree when necessary.
  */
-UCLASS()
-class UVEDITORTOOLS_API UUVToolAABBTreeStorage : public UUVToolContextObject
+UCLASS(MinimalAPI)
+class UUVToolAABBTreeStorage : public UUVToolContextObject
 {
 	GENERATED_BODY()
 public:
@@ -325,20 +327,20 @@ public:
 	using FDynamicMeshAABBTree3 = UE::Geometry::FDynamicMeshAABBTree3;
 
 
-	void Set(FDynamicMesh3* MeshKey, TSharedPtr<FDynamicMeshAABBTree3> Tree,
+	UE_API void Set(FDynamicMesh3* MeshKey, TSharedPtr<FDynamicMeshAABBTree3> Tree,
 		UUVEditorToolMeshInput* InputObject);
 
-	TSharedPtr<FDynamicMeshAABBTree3> Get(FDynamicMesh3* MeshKey);
+	UE_API TSharedPtr<FDynamicMeshAABBTree3> Get(FDynamicMesh3* MeshKey);
 
-	void Remove(FDynamicMesh3* MeshKey);
+	UE_API void Remove(FDynamicMesh3* MeshKey);
 
-	void RemoveByPredicate(TUniqueFunction<
+	UE_API void RemoveByPredicate(TUniqueFunction<
 		bool(const FDynamicMesh3* Mesh, TWeakObjectPtr<UUVEditorToolMeshInput> InputObject,
 			TSharedPtr<FDynamicMeshAABBTree3> Tree)> Predicate);
 
-	void Empty();
+	UE_API void Empty();
 
-	virtual void Shutdown() override;
+	UE_API virtual void Shutdown() override;
 
 protected:
 
@@ -353,8 +355,8 @@ protected:
 * populate other settings menus, such as the Display menu.
 */
 
-UCLASS()
-class UVEDITORTOOLS_API UUVEditorToolPropertiesAPI : public UUVToolContextObject
+UCLASS(MinimalAPI)
+class UUVEditorToolPropertiesAPI : public UUVToolContextObject
 {
 	GENERATED_BODY()
 public:
@@ -382,4 +384,6 @@ protected:
 	UPROPERTY()
 	TObjectPtr<UObject> ToolDisplayProperties = nullptr;
 
-};
\ No newline at end of file
+};
+
+#undef UE_API
diff --git a/Engine/Plugins/Editor/UVEditor/Source/UVEditorTools/Public/ContextObjects/UVToolViewportButtonsAPI.h b/Engine/Plugins/Editor/UVEditor/Source/UVEditorTools/Public/ContextObjects/UVToolViewportButtonsAPI.h
--- a/Engine/Plugins/Editor/UVEditor/Source/UVEditorTools/Public/ContextObjects/UVToolViewportButtonsAPI.h
+++ b/Engine/Plugins/Editor/UVEditor/Source/UVEditorTools/Public/ContextObjects/UVToolViewportButtonsAPI.h
@@ -11,8 +11,8 @@
 /**
  * Allows tools to interact with buttons in the viewport
  */
-UCLASS()
-class UVEDITORTOOLS_API UUVToolViewportButtonsAPI : public UUVToolContextObject
+UCLASS(MinimalAPI)
+class UUVToolViewportButtonsAPI : public UUVToolContextObject
 {
 	GENERATED_BODY()
 public:
@@ -158,4 +158,4 @@ protected:
 	EGizmoMode GizmoMode = EGizmoMode::Select;
 	bool bSelectionButtonsEnabled = false;
 	ESelectionMode SelectionMode = ESelectionMode::Island;
-};
\ No newline at end of file
+};
diff --git a/Engine/Plugins/Editor/UVEditor/Source/UVEditorTools/Public/Drawing/BasicLineSetComponent.h b/Engine/Plugins/Editor/UVEditor/Source/UVEditorTools/Public/Drawing/BasicLineSetComponent.h
--- a/Engine/Plugins/Editor/UVEditor/Source/UVEditorTools/Public/Drawing/BasicLineSetComponent.h
+++ b/Engine/Plugins/Editor/UVEditor/Source/UVEditorTools/Public/Drawing/BasicLineSetComponent.h
@@ -7,6 +7,8 @@
 #include "BasicElementSet.h"
 #include "BasicLineSetComponent.generated.h"
 
+#define UE_API UVEDITORTOOLS_API
+
 class FBasicLineSetSceneProxy;
 
 
@@ -26,29 +28,29 @@ class FBasicLineSetSceneProxy;
  */
 
 /** Base class for component specific functionality independent of the type of line stored in the component. */
-UCLASS()
-class UVEDITORTOOLS_API UBasicLineSetComponentBase : public UMeshComponent
+UCLASS(MinimalAPI)
+class UBasicLineSetComponentBase : public UMeshComponent
 {
 	GENERATED_BODY()
 
 public:
 
-	UBasicLineSetComponentBase();
+	UE_API UBasicLineSetComponentBase();
 
 	/** Clears the component state and marks component as dirty. */
-	void ClearComponent();
+	UE_API void ClearComponent();
 
 	/** Specify material that handles lines. */
-	void SetLineMaterial(UMaterialInterface* InLineMaterial);
+	UE_API void SetLineMaterial(UMaterialInterface* InLineMaterial);
 
 	/** Set per line material parameters that are uniform for all lines. */
-	void SetLineSetParameters(FColor InColor, float InSize, float InDepthBias);
+	UE_API void SetLineSetParameters(FColor InColor, float InSize, float InDepthBias);
 
 	/** Set per line material Line Color parameter that is uniform for all lines */
-	void SetColor(FColor InColor);
+	UE_API void SetColor(FColor InColor);
 
 	/** Set per line material Line Size/Thickness parameter that is uniform for all lines */
-	void SetLineThickness(float InSize);
+	UE_API void SetLineThickness(float InSize);
 
 protected:
 
@@ -105,8 +107,8 @@ protected:
 };
 
 /** Instantiation of a basic line set component in 2D using FVector2f for line positions. */
-UCLASS()
-class UVEDITORTOOLS_API UBasic2DLineSetComponent final : public UBasicLineSetComponentBase, public TBasicElementSet<FVector2f, 2>
+UCLASS(MinimalAPI)
+class UBasic2DLineSetComponent final : public UBasicLineSetComponentBase, public TBasicElementSet<FVector2f, 2>
 {
 	GENERATED_BODY()
 
@@ -121,7 +123,7 @@ public:
 private:
 
 	//~ UPrimitiveComponent Interface.
-	virtual FPrimitiveSceneProxy * CreateSceneProxy() override;
+	UE_API virtual FPrimitiveSceneProxy * CreateSceneProxy() override;
 
 	//~ USceneComponent Interface.
 	virtual FBoxSphereBounds CalcBounds(const FTransform& LocalToWorld) const override
@@ -134,8 +136,8 @@ private:
 };
 
 /** Instantiation of a basic line set component in 3D using FVector3f for line positions. */
-UCLASS()
-class UVEDITORTOOLS_API UBasic3DLineSetComponent final : public UBasicLineSetComponentBase, public TBasicElementSet<FVector3f, 2>
+UCLASS(MinimalAPI)
+class UBasic3DLineSetComponent final : public UBasicLineSetComponentBase, public TBasicElementSet<FVector3f, 2>
 {
 	GENERATED_BODY()
 
@@ -150,7 +152,7 @@ public:
 private:
 
 	//~ UPrimitiveComponent Interface.
-	virtual FPrimitiveSceneProxy * CreateSceneProxy() override;
+	UE_API virtual FPrimitiveSceneProxy * CreateSceneProxy() override;
 
 	//~ USceneComponent Interface.
 	virtual FBoxSphereBounds CalcBounds(const FTransform& LocalToWorld) const override
@@ -161,3 +163,5 @@ private:
 
 	friend class FBasicLineSetSceneProxy;
 };
+
+#undef UE_API
diff --git a/Engine/Plugins/Editor/UVEditor/Source/UVEditorTools/Public/Drawing/BasicPointSetComponent.h b/Engine/Plugins/Editor/UVEditor/Source/UVEditorTools/Public/Drawing/BasicPointSetComponent.h
--- a/Engine/Plugins/Editor/UVEditor/Source/UVEditorTools/Public/Drawing/BasicPointSetComponent.h
+++ b/Engine/Plugins/Editor/UVEditor/Source/UVEditorTools/Public/Drawing/BasicPointSetComponent.h
@@ -7,6 +7,8 @@
 #include "BasicElementSet.h"
 #include "BasicPointSetComponent.generated.h"
 
+#define UE_API UVEDITORTOOLS_API
+
 class FBasicPointSetSceneProxy;
 
 
@@ -26,29 +28,29 @@ class FBasicPointSetSceneProxy;
  */
 
 /** Base class for component specific functionality independent of the type of point stored in the component. */
-UCLASS()
-class UVEDITORTOOLS_API UBasicPointSetComponentBase : public UMeshComponent
+UCLASS(MinimalAPI)
+class UBasicPointSetComponentBase : public UMeshComponent
 {
 	GENERATED_BODY()
 
 public:
 
-	UBasicPointSetComponentBase();
+	UE_API UBasicPointSetComponentBase();
 
 	/** Clears the component state and marks component as dirty. */
-	void ClearComponent();
+	UE_API void ClearComponent();
 
 	/** Specify material that handles points. */
-	void SetPointMaterial(UMaterialInterface* InPointMaterial);
+	UE_API void SetPointMaterial(UMaterialInterface* InPointMaterial);
 
 	/** Set per point material parameters that are uniform for all points. */
-	void SetPointSetParameters(FColor InColor, float InSize, float InDepthBias);
+	UE_API void SetPointSetParameters(FColor InColor, float InSize, float InDepthBias);
 
 	/** Set per point material Color parameter that is uniform for all points. */
-	void SetColor(const FColor InColor);
+	UE_API void SetColor(const FColor InColor);
 
 	/** Set per point material Size parameter tat is uniform for all points. */
-	void SetPointSize(const float InSize);
+	UE_API void SetPointSize(const float InSize);
 
 protected:
 
@@ -105,8 +107,8 @@ protected:
 };
 
 /** Instantiation of a basic point set component in 2D using FVector2f for point positions. */
-UCLASS()
-class UVEDITORTOOLS_API UBasic2DPointSetComponent final : public UBasicPointSetComponentBase, public TBasicElementSet<FVector2f, 1>
+UCLASS(MinimalAPI)
+class UBasic2DPointSetComponent final : public UBasicPointSetComponentBase, public TBasicElementSet<FVector2f, 1>
 {
 	GENERATED_BODY()
 
@@ -121,7 +123,7 @@ public:
 private:
 
 	//~ UPrimitiveComponent Interface.
-	virtual FPrimitiveSceneProxy* CreateSceneProxy() override;
+	UE_API virtual FPrimitiveSceneProxy* CreateSceneProxy() override;
 
 	//~ USceneComponent Interface.
 	virtual FBoxSphereBounds CalcBounds(const FTransform& LocalToWorld) const override
@@ -134,8 +136,8 @@ private:
 };
 
 /** Instantiation of a basic point set component in 3D using FVector3f for point positions. */
-UCLASS()
-class UVEDITORTOOLS_API UBasic3DPointSetComponent final : public UBasicPointSetComponentBase, public TBasicElementSet<FVector3f, 1>
+UCLASS(MinimalAPI)
+class UBasic3DPointSetComponent final : public UBasicPointSetComponentBase, public TBasicElementSet<FVector3f, 1>
 {
 	GENERATED_BODY()
 
@@ -150,7 +152,7 @@ public:
 private:
 
 	//~ UPrimitiveComponent Interface.
-	virtual FPrimitiveSceneProxy* CreateSceneProxy() override;
+	UE_API virtual FPrimitiveSceneProxy* CreateSceneProxy() override;
 
 	//~ USceneComponent Interface.
 	virtual FBoxSphereBounds CalcBounds(const FTransform& LocalToWorld) const override
@@ -161,3 +163,5 @@ private:
 
 	friend class FBasicPointSetSceneProxy;
 };
+
+#undef UE_API
diff --git a/Engine/Plugins/Editor/UVEditor/Source/UVEditorTools/Public/Drawing/BasicTriangleSetComponent.h b/Engine/Plugins/Editor/UVEditor/Source/UVEditorTools/Public/Drawing/BasicTriangleSetComponent.h
--- a/Engine/Plugins/Editor/UVEditor/Source/UVEditorTools/Public/Drawing/BasicTriangleSetComponent.h
+++ b/Engine/Plugins/Editor/UVEditor/Source/UVEditorTools/Public/Drawing/BasicTriangleSetComponent.h
@@ -7,6 +7,8 @@
 #include "BasicElementSet.h"
 #include "BasicTriangleSetComponent.generated.h"
 
+#define UE_API UVEDITORTOOLS_API
+
 class FBasicTriangleSetSceneProxy;
 
 
@@ -26,23 +28,23 @@ class FBasicTriangleSetSceneProxy;
  */
 
  /** Base class for component specific functionality independent of the type of line stored in the component. */
-UCLASS()
-class UVEDITORTOOLS_API UBasicTriangleSetComponentBase : public UMeshComponent
+UCLASS(MinimalAPI)
+class UBasicTriangleSetComponentBase : public UMeshComponent
 {
 	GENERATED_BODY()
 
 public:
 
-	UBasicTriangleSetComponentBase();
+	UE_API UBasicTriangleSetComponentBase();
 
 	/** Clears the component state and marks component as dirty. */
-	void ClearComponent();
+	UE_API void ClearComponent();
 
 	/** Specify material that handles lines. */
-	void SetTriangleMaterial(UMaterialInterface* InTriangleMaterial);
+	UE_API void SetTriangleMaterial(UMaterialInterface* InTriangleMaterial);
 
 	/** Set per line material parameters that are uniform for all lines. */
-	void SetTriangleSetParameters(FColor InColor, const FVector3f& NormalIn);
+	UE_API void SetTriangleSetParameters(FColor InColor, const FVector3f& NormalIn);
 
 protected:
 
@@ -83,8 +85,8 @@ protected:
 };
 
 /** Instantiation of a basic line set component in 2D using FVector2f for line positions. */
-UCLASS()
-class UVEDITORTOOLS_API UBasic2DTriangleSetComponent final : public UBasicTriangleSetComponentBase, public TBasicElementSet<FVector2f, 3>
+UCLASS(MinimalAPI)
+class UBasic2DTriangleSetComponent final : public UBasicTriangleSetComponentBase, public TBasicElementSet<FVector2f, 3>
 {
 	GENERATED_BODY()
 
@@ -99,7 +101,7 @@ public:
 private:
 
 	//~ UPrimitiveComponent Interface.
-	virtual FPrimitiveSceneProxy * CreateSceneProxy() override;
+	UE_API virtual FPrimitiveSceneProxy * CreateSceneProxy() override;
 
 	//~ USceneComponent Interface.
 	virtual FBoxSphereBounds CalcBounds(const FTransform& LocalToWorld) const override
@@ -112,8 +114,8 @@ private:
 };
 
 /** Instantiation of a basic line set component in 3D using FVector3f for line positions. */
-UCLASS()
-class UVEDITORTOOLS_API UBasic3DTriangleSetComponent final : public UBasicTriangleSetComponentBase, public TBasicElementSet<FVector3f, 3>
+UCLASS(MinimalAPI)
+class UBasic3DTriangleSetComponent final : public UBasicTriangleSetComponentBase, public TBasicElementSet<FVector3f, 3>
 {
 	GENERATED_BODY()
 
@@ -128,7 +130,7 @@ public:
 private:
 
 	//~ UPrimitiveComponent Interface.
-	virtual FPrimitiveSceneProxy * CreateSceneProxy() override;
+	UE_API virtual FPrimitiveSceneProxy * CreateSceneProxy() override;
 
 	//~ USceneComponent Interface.
 	virtual FBoxSphereBounds CalcBounds(const FTransform& LocalToWorld) const override
@@ -139,3 +141,5 @@ private:
 
 	friend class FBasicTriangleSetSceneProxy;
 };
+
+#undef UE_API
diff --git a/Engine/Plugins/Editor/UVEditor/Source/UVEditorTools/Public/Operators/UVEditorUVTransformOp.h b/Engine/Plugins/Editor/UVEditor/Source/UVEditorTools/Public/Operators/UVEditorUVTransformOp.h
--- a/Engine/Plugins/Editor/UVEditor/Source/UVEditorTools/Public/Operators/UVEditorUVTransformOp.h
+++ b/Engine/Plugins/Editor/UVEditor/Source/UVEditorTools/Public/Operators/UVEditorUVTransformOp.h
@@ -11,6 +11,8 @@
 
 #include "UVEditorUVTransformOp.generated.h"
 
+#define UE_API UVEDITORTOOLS_API
+
 class UUVTransformProperties;
 class UUVEditorTransformTool;
 
@@ -110,8 +112,8 @@ enum class EUVEditorAlignDistributeGroupingMode : uint8
 /**
  * UV Transform Settings
  */
-UCLASS()
-class UVEDITORTOOLS_API UUVEditorUVTransformPropertiesBase : public UInteractiveToolPropertySet
+UCLASS(MinimalAPI)
+class UUVEditorUVTransformPropertiesBase : public UInteractiveToolPropertySet
 {
 	GENERATED_BODY()
 public:
@@ -132,8 +134,8 @@ public:
 /**
  * UV Transform Settings
  */
-UCLASS()
-class UVEDITORTOOLS_API UUVEditorUVTransformProperties : public UUVEditorUVTransformPropertiesBase
+UCLASS(MinimalAPI)
+class UUVEditorUVTransformProperties : public UUVEditorUVTransformPropertiesBase
 {
 	GENERATED_BODY()
 
@@ -189,8 +191,8 @@ public:
  * This is designed to be used in the future where we want to provide a quick translate "tool" when no other tools are running.
  * 
  */
-UCLASS()
-class UVEDITORTOOLS_API UUVEditorUVQuickTransformProperties : public UUVEditorUVTransformProperties
+UCLASS(MinimalAPI)
+class UUVEditorUVQuickTransformProperties : public UUVEditorUVTransformProperties
 {
 	GENERATED_BODY()
 };
@@ -198,8 +200,8 @@ class UVEDITORTOOLS_API UUVEditorUVQuickTransformProperties : public UUVEditorUV
 /**
  * UV Align Settings
  */
-UCLASS()
-class UVEDITORTOOLS_API UUVEditorUVAlignProperties : public UUVEditorUVTransformPropertiesBase
+UCLASS(MinimalAPI)
+class UUVEditorUVAlignProperties : public UUVEditorUVTransformPropertiesBase
 {
 	GENERATED_BODY()
 
@@ -227,8 +229,8 @@ public:
 /**
  * UV Distribute Settings
  */
-UCLASS()
-class UVEDITORTOOLS_API UUVEditorUVDistributeProperties : public UUVEditorUVTransformPropertiesBase
+UCLASS(MinimalAPI)
+class UUVEditorUVDistributeProperties : public UUVEditorUVTransformPropertiesBase
 {
 	GENERATED_BODY()
 public:
@@ -326,7 +328,7 @@ namespace Geometry
 		IndividualVertices
 	};
 
-class UVEDITORTOOLS_API FUVEditorUVTransformBaseOp : public FDynamicMeshOperator
+class FUVEditorUVTransformBaseOp : public FDynamicMeshOperator
 {
 public:
 	virtual ~FUVEditorUVTransformBaseOp() {}
@@ -336,20 +338,20 @@ public:
 	int UVLayerIndex = 0;
 	EUVEditorAlignDistributeGroupingModeBackend GroupingMode;
 
-	void SetTransform(const FTransformSRT3d& Transform);
-	void SetSelection(const TSet<int32>& EdgeSelection, const TSet<int32>& VertexSelection);
+	UE_API void SetTransform(const FTransformSRT3d& Transform);
+	UE_API void SetSelection(const TSet<int32>& EdgeSelection, const TSet<int32>& VertexSelection);
 
-	virtual void CalculateResult(FProgressCancel* Progress) override;
+	UE_API virtual void CalculateResult(FProgressCancel* Progress) override;
 	const TArray<FVector2D>& GetPivotLocations() const { return VisualizationPivots; }
 
 protected:
 	virtual void HandleTransformationOp(FProgressCancel* Progress) = 0;
 
-	FVector2f GetAlignmentPointFromBoundingBoxAndDirection(EUVEditorAlignDirectionBackend Direction, const FAxisAlignedBox2d& BoundingBox);
-	FVector2f GetAlignmentPointFromUDIMAndDirection(EUVEditorAlignDirectionBackend Direction, FVector2i UDIMTile);
-	void RebuildBoundingBoxes();
-	void CollectTransformElements();
-	void SortComponentsByBoundingBox(bool bSortXThenY = true);
+	UE_API FVector2f GetAlignmentPointFromBoundingBoxAndDirection(EUVEditorAlignDirectionBackend Direction, const FAxisAlignedBox2d& BoundingBox);
+	UE_API FVector2f GetAlignmentPointFromUDIMAndDirection(EUVEditorAlignDirectionBackend Direction, FVector2i UDIMTile);
+	UE_API void RebuildBoundingBoxes();
+	UE_API void CollectTransformElements();
+	UE_API void SortComponentsByBoundingBox(bool bSortXThenY = true);
 
 	TMap<int32, int32> ElementToComponent;
 	FDynamicMeshUVOverlay* ActiveUVLayer;
@@ -367,7 +369,7 @@ protected:
 
 /**
  */
-class UVEDITORTOOLS_API FUVEditorUVTransformOp : public FUVEditorUVTransformBaseOp
+class FUVEditorUVTransformOp : public FUVEditorUVTransformBaseOp
 {
 public:
 	virtual ~FUVEditorUVTransformOp() {}
@@ -385,11 +387,11 @@ public:
 	FVector2D ManualPivot = FVector2D(0, 0);
 
 protected:
-	virtual void HandleTransformationOp(FProgressCancel* Progress) override;
-	FVector2f GetPivotFromMode(int32 ElementID, EUVEditorPivotTypeBackend Mode);
+	UE_API virtual void HandleTransformationOp(FProgressCancel* Progress) override;
+	UE_API FVector2f GetPivotFromMode(int32 ElementID, EUVEditorPivotTypeBackend Mode);
 };
 
-class UVEDITORTOOLS_API FUVEditorUVAlignOp : public FUVEditorUVTransformBaseOp
+class FUVEditorUVAlignOp : public FUVEditorUVTransformBaseOp
 {
 public:
 	virtual ~FUVEditorUVAlignOp() {}
@@ -399,13 +401,13 @@ public:
 	FVector2D ManualAnchor = FVector2D(0, 0);
 
 protected:
-	virtual void HandleTransformationOp(FProgressCancel* Progress) override;
+	UE_API virtual void HandleTransformationOp(FProgressCancel* Progress) override;
 
 
 
 };
 
-class UVEDITORTOOLS_API FUVEditorUVDistributeOp : public FUVEditorUVTransformBaseOp
+class FUVEditorUVDistributeOp : public FUVEditorUVTransformBaseOp
 {
 public:
 	virtual ~FUVEditorUVDistributeOp() {}
@@ -416,7 +418,7 @@ public:
 	float ManualSpacing;
 
 protected:
-	virtual void HandleTransformationOp(FProgressCancel* Progress) override;
+	UE_API virtual void HandleTransformationOp(FProgressCancel* Progress) override;
 
 };
 
@@ -430,14 +432,14 @@ protected:
  * Inherits from UObject so that it can hold a strong pointer to the settings UObject, which
  * needs to be a UObject to be displayed in the details panel.
  */
-UCLASS()
-class UVEDITORTOOLS_API UUVEditorUVTransformOperatorFactory : public UObject, public UE::Geometry::IDynamicMeshOperatorFactory
+UCLASS(MinimalAPI)
+class UUVEditorUVTransformOperatorFactory : public UObject, public UE::Geometry::IDynamicMeshOperatorFactory
 {
 	GENERATED_BODY()
 
 public:
 	// IDynamicMeshOperatorFactory API
-	virtual TUniquePtr<UE::Geometry::FDynamicMeshOperator> MakeNewOperator() override;
+	UE_API virtual TUniquePtr<UE::Geometry::FDynamicMeshOperator> MakeNewOperator() override;
 
 	EUVEditorUVTransformType TransformType;
 
@@ -450,4 +452,6 @@ public:
 	TSharedPtr<UE::Geometry::FDynamicMesh3, ESPMode::ThreadSafe> OriginalMesh;
 	TUniqueFunction<int32()> GetSelectedUVChannel = []() { return 0; };
 	FTransform TargetTransform;
-};
\ No newline at end of file
+};
+
+#undef UE_API
diff --git a/Engine/Plugins/Editor/UVEditor/Source/UVEditorTools/Public/Selection/UVEditorMeshSelectionMechanic.h b/Engine/Plugins/Editor/UVEditor/Source/UVEditorTools/Public/Selection/UVEditorMeshSelectionMechanic.h
--- a/Engine/Plugins/Editor/UVEditor/Source/UVEditorTools/Public/Selection/UVEditorMeshSelectionMechanic.h
+++ b/Engine/Plugins/Editor/UVEditor/Source/UVEditorTools/Public/Selection/UVEditorMeshSelectionMechanic.h
@@ -16,6 +16,8 @@
 
 #include "UVEditorMeshSelectionMechanic.generated.h"
 
+#define UE_API UVEDITORTOOLS_API
+
 class APreviewGeometryActor;
 struct FCameraRectangle;
 class ULineSetComponent;
@@ -30,8 +32,8 @@ class ULocalMouseHoverBehavior;
  * Mechanic for selecting elements of a dynamic mesh in the UV editor. Interacts
  * heavily with UUVToolSelectionAPI, which actually stores selections.
  */
-UCLASS()
-class UVEDITORTOOLS_API UUVEditorMeshSelectionMechanic : public UInteractionMechanic
+UCLASS(MinimalAPI)
+class UUVEditorMeshSelectionMechanic : public UInteractionMechanic
 {
 	GENERATED_BODY()
 
@@ -41,17 +43,17 @@ public:
 
 	virtual ~UUVEditorMeshSelectionMechanic() {}
 
-	virtual void Setup(UInteractiveTool* ParentTool) override;
-	virtual void Shutdown() override;
+	UE_API virtual void Setup(UInteractiveTool* ParentTool) override;
+	UE_API virtual void Shutdown() override;
 
 	// Initialization functions.
 	// The selection API is provided as a parameter rather than being grabbed out of the context
 	// store mainly because UVToolSelectionAPI itself sets up a selection mechanic, and is not 
 	// yet in the context store when it does this. 
-	void Initialize(UWorld* World, UWorld* LivePreviewWorld, UUVToolSelectionAPI* SelectionAPI);
-	void SetTargets(const TArray<TObjectPtr<UUVEditorToolMeshInput>>& TargetsIn);
+	UE_API void Initialize(UWorld* World, UWorld* LivePreviewWorld, UUVToolSelectionAPI* SelectionAPI);
+	UE_API void SetTargets(const TArray<TObjectPtr<UUVEditorToolMeshInput>>& TargetsIn);
 
-	void SetIsEnabled(bool bIsEnabled);
+	UE_API void SetIsEnabled(bool bIsEnabled);
 	bool IsEnabled() { return bIsEnabled; };
 
 	struct FRaycastResult
@@ -66,37 +68,37 @@ public:
 	 *  keeps aabb trees for them. The mechanic does not need to be enabled to use this (and typically
 	 *  will not be, if a tool is doing its own raycasting).
 	 */
-	bool RaycastCanonicals(const FRay& WorldRay, bool bRaycastIsForUnwrap,
+	UE_API bool RaycastCanonicals(const FRay& WorldRay, bool bRaycastIsForUnwrap,
 		bool bPreferSelected, FRaycastResult& HitOut) const;
 
-	TArray<FUVToolSelection> GetAllCanonicalTrianglesInUnwrapRadius(const FVector2d& UnwrapWorldHitPoint, double Radius) const;
+	UE_API TArray<FUVToolSelection> GetAllCanonicalTrianglesInUnwrapRadius(const FVector2d& UnwrapWorldHitPoint, double Radius) const;
 
-	void SetShowHoveredElements(bool bShow);
+	UE_API void SetShowHoveredElements(bool bShow);
 
 	using ESelectionMode = UUVToolSelectionAPI::EUVEditorSelectionMode;
 	using FModeChangeOptions = UUVToolSelectionAPI::FSelectionMechanicModeChangeOptions;
 	/**
 	 * Sets selection mode for the mechanic.
 	 */
-	void SetSelectionMode(ESelectionMode TargetMode,
+	UE_API void SetSelectionMode(ESelectionMode TargetMode,
 		const FModeChangeOptions& Options = FModeChangeOptions());
 	
-	virtual void Render(IToolsContextRenderAPI* RenderAPI) override;
-	virtual void DrawHUD(FCanvas* Canvas, IToolsContextRenderAPI* RenderAPI);
-	virtual void LivePreviewRender(IToolsContextRenderAPI* RenderAPI);
-	virtual void LivePreviewDrawHUD(FCanvas* Canvas, IToolsContextRenderAPI* RenderAPI);
+	UE_API virtual void Render(IToolsContextRenderAPI* RenderAPI) override;
+	UE_API virtual void DrawHUD(FCanvas* Canvas, IToolsContextRenderAPI* RenderAPI);
+	UE_API virtual void LivePreviewRender(IToolsContextRenderAPI* RenderAPI);
+	UE_API virtual void LivePreviewDrawHUD(FCanvas* Canvas, IToolsContextRenderAPI* RenderAPI);
 	// IClickBehaviorTarget implementation
-	virtual FInputRayHit IsHitByClick(const FInputDeviceRay& ClickPos, bool bSourceIsLivePreview);
-	virtual void OnClicked(const FInputDeviceRay& ClickPos, bool bSourceIsLivePreview);
+	UE_API virtual FInputRayHit IsHitByClick(const FInputDeviceRay& ClickPos, bool bSourceIsLivePreview);
+	UE_API virtual void OnClicked(const FInputDeviceRay& ClickPos, bool bSourceIsLivePreview);
 
 	// IModifierToggleBehaviorTarget implementation
-	virtual void OnUpdateModifierState(int ModifierID, bool bIsOn);
+	UE_API virtual void OnUpdateModifierState(int ModifierID, bool bIsOn);
 
 	// IHoverBehaviorTarget implementation
-	virtual FInputRayHit BeginHoverSequenceHitTest(const FInputDeviceRay& PressPos, bool bSourceIsLivePreview);
-	virtual void OnBeginHover(const FInputDeviceRay& DevicePos);
-	virtual bool OnUpdateHover(const FInputDeviceRay& DevicePos, bool bSourceIsLivePreview);
-	virtual void OnEndHover();
+	UE_API virtual FInputRayHit BeginHoverSequenceHitTest(const FInputDeviceRay& PressPos, bool bSourceIsLivePreview);
+	UE_API virtual void OnBeginHover(const FInputDeviceRay& DevicePos);
+	UE_API virtual bool OnUpdateHover(const FInputDeviceRay& DevicePos, bool bSourceIsLivePreview);
+	UE_API virtual void OnEndHover();
 
 	/**
 	 * Broadcasted whenever the marquee mechanic rectangle is changed, since these changes
@@ -162,7 +164,7 @@ protected:
 	TWeakObjectPtr<UPointSetComponent> LivePreviewHoverPointSet = nullptr;
 
 	// Should be the same as the mode-level targets array, indexed by AssetID
-	TSharedPtr<FDynamicMeshAABBTree3> GetMeshSpatial(int32 TargetId, bool bUseUnwrap) const;
+	UE_API TSharedPtr<FDynamicMeshAABBTree3> GetMeshSpatial(int32 TargetId, bool bUseUnwrap) const;
 	TArray<TObjectPtr<UUVEditorToolMeshInput>> Targets;
 	TArray<TSharedPtr<FDynamicMeshAABBTree3>> UnwrapMeshSpatials; // 1:1 with Targets
 	TArray<TSharedPtr<FDynamicMeshAABBTree3>> AppliedMeshSpatials; // 1:1 with Targets
@@ -171,9 +173,9 @@ protected:
 	bool bIsEnabled = false;
 	bool bShowHoveredElements = true;
 
-	bool GetHitTid(const FInputDeviceRay& ClickPos, int32& TidOut,
+	UE_API bool GetHitTid(const FInputDeviceRay& ClickPos, int32& TidOut,
 		int32& AssetIDOut, bool bUseUnwrap, int32* ExistingSelectionObjectIndexOut = nullptr);
-	void ModifyExistingSelection(TSet<int32>& SelectionSetToModify, const TArray<int32>& SelectedIDs);
+	UE_API void ModifyExistingSelection(TSet<int32>& SelectionSetToModify, const TArray<int32>& SelectedIDs);
 
 	FViewCameraState CameraState;
 
@@ -189,9 +191,9 @@ protected:
 	bool ShouldRestartSelection() const { return !bCtrlToggle && !bShiftToggle; }
 
 	// For marquee mechanic
-	void OnDragRectangleStarted();
-	void OnDragRectangleChanged(const FCameraRectangle& CurrentRectangle, bool bSourceIsLivePreview);
-	void OnDragRectangleFinished(const FCameraRectangle& Rectangle, bool bCancelled, bool bSourceIsLivePreview);
+	UE_API void OnDragRectangleStarted();
+	UE_API void OnDragRectangleChanged(const FCameraRectangle& CurrentRectangle, bool bSourceIsLivePreview);
+	UE_API void OnDragRectangleFinished(const FCameraRectangle& Rectangle, bool bCancelled, bool bSourceIsLivePreview);
 
 	TArray<FUVToolSelection> PreDragSelections;
 	TArray<FUVToolSelection> PreDragUnsetSelections;
@@ -201,3 +203,4 @@ protected:
 	TArray<const FUVToolSelection*> AssetIDToPreDragUnsetSelection;
 };
 
+#undef UE_API
diff --git a/Engine/Plugins/Editor/UVEditor/Source/UVEditorTools/Public/Selection/UVToolSelection.h b/Engine/Plugins/Editor/UVEditor/Source/UVEditorTools/Public/Selection/UVToolSelection.h
--- a/Engine/Plugins/Editor/UVEditor/Source/UVEditorTools/Public/Selection/UVToolSelection.h
+++ b/Engine/Plugins/Editor/UVEditor/Source/UVEditorTools/Public/Selection/UVToolSelection.h
@@ -9,13 +9,15 @@
 #include "DynamicMesh/DynamicMesh3.h"
 #include "Util/IndexUtil.h"
 
+#define UE_API UVEDITORTOOLS_API
+
 namespace UE {
 namespace Geometry {
 
 /**
  * Class that represents a selection in the canonical unwrap of a UV editor input object.
  */
-class UVEDITORTOOLS_API FUVToolSelection
+class FUVToolSelection
 {
 public:
 	enum class EType
@@ -36,9 +38,9 @@ public:
 		StableEdgeIDs.Reset();
 	}
 
-	void SelectAll(const FDynamicMesh3& Mesh, EType TypeIn);
+	UE_API void SelectAll(const FDynamicMesh3& Mesh, EType TypeIn);
 
-	bool IsAllSelected(const FDynamicMesh3& Mesh) const;
+	UE_API bool IsAllSelected(const FDynamicMesh3& Mesh) const;
 
 	bool IsEmpty() const
 	{
@@ -64,32 +66,32 @@ public:
 		return !(*this == Other);
 	}
 
-	void SaveStableEdgeIdentifiers(const FDynamicMesh3& Mesh);
+	UE_API void SaveStableEdgeIdentifiers(const FDynamicMesh3& Mesh);
 
-	void RestoreFromStableEdgeIdentifiers(const FDynamicMesh3& Mesh);
+	UE_API void RestoreFromStableEdgeIdentifiers(const FDynamicMesh3& Mesh);
 
-	bool AreElementsPresentInMesh(const FDynamicMesh3& Mesh) const;
+	UE_API bool AreElementsPresentInMesh(const FDynamicMesh3& Mesh) const;
 
-	FAxisAlignedBox3d ToBoundingBox(const FDynamicMesh3& Mesh, const FTransform Transform = FTransform3d::Identity) const;
+	UE_API FAxisAlignedBox3d ToBoundingBox(const FDynamicMesh3& Mesh, const FTransform Transform = FTransform3d::Identity) const;
 
 	/**
 	 * Utility method to construct a new selection from an existing reference selection, only with a new geometry type.
 	 */
-	FUVToolSelection GetConvertedSelection(const FDynamicMesh3& Mesh, FUVToolSelection::EType ExpectedSelectionType) const;
+	UE_API FUVToolSelection GetConvertedSelection(const FDynamicMesh3& Mesh, FUVToolSelection::EType ExpectedSelectionType) const;
 
 	/** 
 	* Utility method to construct a new selection from an existing reference Applied mesh selection, only for the equivalent Unwrap mesh geometry.
 	* 
 	* If the target for this selection is currently invalid or unset, this will return an empty selection.
 	*/
-	FUVToolSelection GetConvertedSelectionForUnwrappedMesh() const;
+	UE_API FUVToolSelection GetConvertedSelectionForUnwrappedMesh() const;
 
 	/**
 	* Utility method to construct a new selection from an existing reference Unwrap mesh selection, only for the equivalent Applied mesh geometry.
 	*
 	* If the target for this selection is currently invalid or unset, this will return an empty selection.
 	*/
-	FUVToolSelection GetConvertedSelectionForAppliedMesh() const;
+	UE_API FUVToolSelection GetConvertedSelectionForAppliedMesh() const;
 
 	/**
 	 * Utility method to convert element ids from those belonging to the unwrapped mesh to those belonging to equivalent ids on the applied mesh
@@ -100,7 +102,7 @@ public:
 	 * @param UVOverlay An overlay representing the current UV mapping linking the unwrapped and applied mesh
 	 * @param IDsIn The ids of the elements from the unwrapped mesh to convert to applied element ids
 	 */
-	static TArray<int32> ConvertUnwrappedElementIdsToAppliedElementIds(EType SelectionMode, const FDynamicMesh3& UnwrapMesh,
+	static UE_API TArray<int32> ConvertUnwrappedElementIdsToAppliedElementIds(EType SelectionMode, const FDynamicMesh3& UnwrapMesh,
 		const FDynamicMesh3& AppliedMesh, const FDynamicMeshUVOverlay& UVOverlay, const TArray<int32>& IDsIn);
 
 	/**
@@ -113,7 +115,7 @@ public:
 	 * @param IDsIn The ids of the elements from the unwrapped mesh to convert to applied element ids
 	 * @param AppliedOnlyElementIds The ids of any remaining elements that have no direct mappings, due to unset UVs, in the unwrapped mesh
 	 */
-	static TArray<int32> ConvertAppliedElementIdsToUnwrappedElementIds(EType SelectionMode, const FDynamicMesh3& AppliedMesh,
+	static UE_API TArray<int32> ConvertAppliedElementIdsToUnwrappedElementIds(EType SelectionMode, const FDynamicMesh3& AppliedMesh,
 		const FDynamicMesh3& UnwrapMesh, const FDynamicMeshUVOverlay& UVOverlay, const TArray<int32>& IDsIn, TArray<int32>& AppliedOnlyElementIds);
 
 protected:
@@ -121,4 +123,6 @@ protected:
 	FMeshEdgesFromTriangleSubIndices StableEdgeIDs;
 };
 
-}} //end namespace UE::Geometry
\ No newline at end of file
+}} //end namespace UE::Geometry
+
+#undef UE_API
diff --git a/Engine/Plugins/Editor/UVEditor/Source/UVEditorTools/Public/Selection/UVToolSelectionAPI.h b/Engine/Plugins/Editor/UVEditor/Source/UVEditorTools/Public/Selection/UVToolSelectionAPI.h
--- a/Engine/Plugins/Editor/UVEditor/Source/UVEditorTools/Public/Selection/UVToolSelectionAPI.h
+++ b/Engine/Plugins/Editor/UVEditor/Source/UVEditorTools/Public/Selection/UVToolSelectionAPI.h
@@ -9,6 +9,8 @@
 
 #include "UVToolSelectionAPI.generated.h"
 
+#define UE_API UVEDITORTOOLS_API
+
 class UUVEditorMechanicAdapterTool;
 class UUVEditorMeshSelectionMechanic;
 class UUVToolSelectionHighlightMechanic;
@@ -36,8 +38,8 @@ ENUM_CLASS_FLAGS(ESelectionChangeTypeFlag);
  * and to enable a selection mechanic in the viewport (to which tools can respond via
  * OnSelectionChanged broadcasts).
  */
-UCLASS()
-class UVEDITORTOOLS_API UUVToolSelectionAPI : public UUVToolContextObject
+UCLASS(MinimalAPI)
+class UUVToolSelectionAPI : public UUVToolContextObject
 {
 	GENERATED_BODY()
 public:
@@ -86,12 +88,12 @@ public:
 	 * @param bAllowInconsistentSelectionTypes If true, allow setting selections with different selection types than any currently set UnsetElement selections.
 	          Primarily useful during undo/redo
 	 */
-	void SetSelections(const TArray<FUVToolSelection>& SelectionsIn, bool bBroadcast, bool bEmitChange);
+	UE_API void SetSelections(const TArray<FUVToolSelection>& SelectionsIn, bool bBroadcast, bool bEmitChange);
 
 	/**
 	 * Equivalent to calling SetSelections with an empty array.
 	 */
-	void ClearSelections(bool bBroadcast, bool bEmitChange);
+	UE_API void ClearSelections(bool bBroadcast, bool bEmitChange);
 
 	/**
 	 * Gets the bounding box center of the current selection. The center point is cached and
@@ -102,7 +104,7 @@ public:
 	 * because it is not guaranteed that the selection api will detect the mesh change (and therefore
 	 * invalidate the cached value) before your call site, due to the order of OnCanonicalModified broadcasts
 	 */
-	FVector3d GetUnwrapSelectionBoundingBoxCenter(bool bForceRecalculate = false);
+	UE_API FVector3d GetUnwrapSelectionBoundingBoxCenter(bool bForceRecalculate = false);
 
 	// These are functions that are used if a selection made in the applied mesh cannot be
 	// converted to an unwrap mesh selection due to unset UV elements (meaning that the selected
@@ -110,10 +112,10 @@ public:
 	// from all those that are able to be converted (which get stored with the usual unwrap selections).
 	// These "unset element" selections are ignored by many tools/actions but operated on 
 	// by some others.
-	bool HaveUnsetElementAppliedMeshSelections() const;
+	UE_API bool HaveUnsetElementAppliedMeshSelections() const;
 	const TArray<FUVToolSelection>& GetUnsetElementAppliedMeshSelections() { return CurrentUnsetSelections; }
-	void SetUnsetElementAppliedMeshSelections(const TArray<FUVToolSelection>& SelectionsIn, bool bBroadcast, bool bEmitChange);
-	void ClearUnsetElementAppliedMeshSelections(bool bBroadcast, bool bEmitChange);
+	UE_API void SetUnsetElementAppliedMeshSelections(const TArray<FUVToolSelection>& SelectionsIn, bool bBroadcast, bool bEmitChange);
+	UE_API void ClearUnsetElementAppliedMeshSelections(bool bBroadcast, bool bEmitChange);
 
 	/**
 	 * Mode of operation for the selection mechanic.
@@ -138,13 +140,13 @@ public:
 	 */
 	const UUVEditorMeshSelectionMechanic* GetSelectionMechanic() const { return SelectionMechanic; }
 
-	void SetSelectionMechanicEnabled(bool bEnabled);
+	UE_API void SetSelectionMechanicEnabled(bool bEnabled);
 
 	struct FSelectionMechanicOptions
 	{
 		bool bShowHoveredElements = true;
 	};
-	void SetSelectionMechanicOptions(const FSelectionMechanicOptions& Options);
+	UE_API void SetSelectionMechanicOptions(const FSelectionMechanicOptions& Options);
 
 	// Options for how a SetSelectionMechanicMode() call is performed
 	struct FSelectionMechanicModeChangeOptions
@@ -172,7 +174,7 @@ public:
 	/**
 	 * Sets the mechanic mode of operation.
 	 */
-	void SetSelectionMechanicMode(EUVEditorSelectionMode Mode, 
+	UE_API void SetSelectionMechanicMode(EUVEditorSelectionMode Mode, 
 		const FSelectionMechanicModeChangeOptions& Options = FSelectionMechanicModeChangeOptions());
 	
 	// Highlighting controls
@@ -180,7 +182,7 @@ public:
 	/**
 	 * Changes the visibility of the highlight.
 	 */
-	void SetHighlightVisible(bool bUnwrapHighlightVisible, bool bAppliedHighlightVisible, bool bRebuild = true);
+	UE_API void SetHighlightVisible(bool bUnwrapHighlightVisible, bool bAppliedHighlightVisible, bool bRebuild = true);
 
 	struct FHighlightOptions
 	{
@@ -204,15 +206,15 @@ public:
 		// it can be welded (in different colors).
 		bool bShowPairedEdgeHighlights = false;
 	};
-	void SetHighlightOptions(const FHighlightOptions& Options);
+	UE_API void SetHighlightOptions(const FHighlightOptions& Options);
 
-	void ClearHighlight(bool bClearForUnwrap = true, bool bClearForAppliedPreview = true);
+	UE_API void ClearHighlight(bool bClearForUnwrap = true, bool bClearForAppliedPreview = true);
 
 	/**
 	 * Build up a highlight of the current selection in the Unwrap, with a given transform
 	 * considered as its start transform.
 	 */
-	void RebuildUnwrapHighlight(const FTransform& StartTransform);
+	UE_API void RebuildUnwrapHighlight(const FTransform& StartTransform);
 
 	/**
 	 * Change the transform of the unwrap highlight without rebuilding it (for cheap movement
@@ -220,11 +222,11 @@ public:
 	 * paired edge highlighting is enabled) will still be rebuilt since their shape may be
 	 * changed by the movement of adjacent edges.
 	 */
-	void SetUnwrapHighlightTransform(const FTransform& NewTransform);
+	UE_API void SetUnwrapHighlightTransform(const FTransform& NewTransform);
 
-	FTransform GetUnwrapHighlightTransform();
+	UE_API FTransform GetUnwrapHighlightTransform();
 
-	void RebuildAppliedPreviewHighlight();
+	UE_API void RebuildAppliedPreviewHighlight();
 
 	/**
 	 * Broadcasted right before a selection change is applied (and therefore before a selection change
@@ -321,31 +323,31 @@ public:
 	 * Preps a selection change transaction, if the user wants more control on what the previous and
 	 * current selection is.
 	 */
-	void BeginChange();
+	UE_API void BeginChange();
 
 	/**
 	 * Ends the active change and emits it via the EmitChangeAPI
 	 */
-	bool EndChangeAndEmitIfModified(bool bBroadcast);
+	UE_API bool EndChangeAndEmitIfModified(bool bBroadcast);
 
 	// Initialization functions:
-	void Initialize(UInteractiveToolManager* ToolManagerIn, UWorld* UnwrapWorld,
+	UE_API void Initialize(UInteractiveToolManager* ToolManagerIn, UWorld* UnwrapWorld,
 		UInputRouter* UnwrapInputRouterIn, UUVToolLivePreviewAPI* LivePreviewAPI,
 		UUVToolEmitChangeAPI* EmitChangeAPIIn);
 	// Should be called after Initialize()
-	void SetTargets(const TArray<TObjectPtr<UUVEditorToolMeshInput>>& TargetsIn);
+	UE_API void SetTargets(const TArray<TObjectPtr<UUVEditorToolMeshInput>>& TargetsIn);
 
 	// Called by the owner of the API
-	virtual void Render(IToolsContextRenderAPI* RenderAPI);
-	virtual void DrawHUD(FCanvas* Canvas, IToolsContextRenderAPI* RenderAPI);
+	UE_API virtual void Render(IToolsContextRenderAPI* RenderAPI);
+	UE_API virtual void DrawHUD(FCanvas* Canvas, IToolsContextRenderAPI* RenderAPI);
 	// Could have had the api place these into the LivePreviewAPI delegates, but decided against
 	// it to align with Render() and DrawHUD() above. The owner will just call these directly.
-	virtual void LivePreviewRender(IToolsContextRenderAPI* RenderAPI);
-	virtual void LivePreviewDrawHUD(FCanvas* Canvas, IToolsContextRenderAPI* RenderAPI);
+	UE_API virtual void LivePreviewRender(IToolsContextRenderAPI* RenderAPI);
+	UE_API virtual void LivePreviewDrawHUD(FCanvas* Canvas, IToolsContextRenderAPI* RenderAPI);
 
 	// UUVToolContextObject
-	virtual void Shutdown() override;
-	virtual void OnToolEnded(UInteractiveTool* DeadTool) override;
+	UE_API virtual void Shutdown() override;
+	UE_API virtual void OnToolEnded(UInteractiveTool* DeadTool) override;
 
 	// Holds the properties of the UV Live Preview Selection (in 3d viewport) for customization
 	struct FLivePreviewSelectionUXSettings
@@ -356,7 +358,7 @@ public:
 	};
 
 	// Sets the UV Live Preview Selection settings
-	void SetLivePreviewSelectionUXSettings(const FLivePreviewSelectionUXSettings& Settings) const;
+	UE_API void SetLivePreviewSelectionUXSettings(const FLivePreviewSelectionUXSettings& Settings) const;
 
 protected:
 	TArray<FUVToolSelection> CurrentSelections;
@@ -423,3 +425,5 @@ public:
 	virtual bool SupportsUnsetElementAppliedMeshSelections() const { return false; }
 
 };
+
+#undef UE_API
diff --git a/Engine/Plugins/Editor/UVEditor/Source/UVEditorTools/Public/Selection/UVToolSelectionHighlightMechanic.h b/Engine/Plugins/Editor/UVEditor/Source/UVEditorTools/Public/Selection/UVToolSelectionHighlightMechanic.h
--- a/Engine/Plugins/Editor/UVEditor/Source/UVEditorTools/Public/Selection/UVToolSelectionHighlightMechanic.h
+++ b/Engine/Plugins/Editor/UVEditor/Source/UVEditorTools/Public/Selection/UVToolSelectionHighlightMechanic.h
@@ -9,6 +9,8 @@
 
 #include "UVToolSelectionHighlightMechanic.generated.h"
 
+#define UE_API UVEDITORTOOLS_API
+
 PREDECLARE_GEOMETRY(class FDynamicMesh3);
 PREDECLARE_GEOMETRY(class FUVToolSelection);
 
@@ -25,8 +27,8 @@ class UWorld;
 /**
  * Mechanic for highlighting mesh elements (usually selection) in the UV editor.
  */
-UCLASS()
-class UVEDITORTOOLS_API UUVToolSelectionHighlightMechanic : public UInteractionMechanic
+UCLASS(MinimalAPI)
+class UUVToolSelectionHighlightMechanic : public UInteractionMechanic
 {
 	GENERATED_BODY()
 
@@ -34,10 +36,10 @@ class UVEDITORTOOLS_API UUVToolSelectionHighlightMechanic : public UInteractionM
 
 public:
 
-	void Initialize(UWorld* UnwrapWorld, UWorld* LivePreviewWorld);
-	void Shutdown() override;
+	UE_API void Initialize(UWorld* UnwrapWorld, UWorld* LivePreviewWorld);
+	UE_API void Shutdown() override;
 
-	void SetIsVisible(bool bUnwrapHighlightVisible, bool bLivePreviewHighlightVisible);
+	UE_API void SetIsVisible(bool bUnwrapHighlightVisible, bool bLivePreviewHighlightVisible);
 
 	/**
 	 * Rebuilds the unwrap mesh highlights.
@@ -47,7 +49,7 @@ public:
 	 * @param bUsePreviews If true, the selected element positions will be gotten from the UnwrapPreview
 	 *   of each selection target rather than UnwrapCanonical.
 	 */
-	void RebuildUnwrapHighlight(const TArray<FUVToolSelection>& Selections, 
+	UE_API void RebuildUnwrapHighlight(const TArray<FUVToolSelection>& Selections, 
 		const FTransform& StartTransform, bool bUsePreviews = false);
 
 	/**
@@ -60,13 +62,13 @@ public:
 	 * @param bUsePreviews If bRebuildStaticPairedEdges is true, this will prefer to use the
 	 *   previews rather than the unwrap canonical.
 	 */
-	void SetUnwrapHighlightTransform(const FTransform& Transform, bool bRebuildStaticPairedEdges = true, 
+	UE_API void SetUnwrapHighlightTransform(const FTransform& Transform, bool bRebuildStaticPairedEdges = true, 
 		bool bUsePreviews = false);
 
 	/**
 	 * Get the current transform of the unwrap highlight.
 	 */
-	FTransform GetUnwrapHighlightTransform();
+	UE_API FTransform GetUnwrapHighlightTransform();
 
 	/**
 	 * Convert the unwrap selections to elements in the applied mesh, then highlight those there.
@@ -74,7 +76,7 @@ public:
 	 * @param bUsePreviews If true, the previews are used rather than the canonicals both for converting
 	 *   over to applied mesh elements and to finding the highlight positions.
 	 */
-	void RebuildAppliedHighlightFromUnwrapSelection(const TArray<FUVToolSelection>& UnwrapSelections, bool bUsePreviews = false);
+	UE_API void RebuildAppliedHighlightFromUnwrapSelection(const TArray<FUVToolSelection>& UnwrapSelections, bool bUsePreviews = false);
 
 	/**
 	 * Adds additional highlighting elements based on selection information. This is used to extend selection
@@ -82,29 +84,29 @@ public:
 	 *
 	 * @param bUsePreviews If true, the previews are used rather than the canonicals for finding the highlight positions.
 	 */
-	void AppendAppliedHighlight(const TArray<FUVToolSelection>& AppliedSelections, bool bUsePreviews = false);
+	UE_API void AppendAppliedHighlight(const TArray<FUVToolSelection>& AppliedSelections, bool bUsePreviews = false);
 
 	/**
 	 * When true, edge selections also highlight the edge pairings (i.e. edges that could weld with the
 	 * selected edges). Note that the call itself controls visibility and the next rebuild call, but 
 	 * won't cause an immediate rebuild on its own.
 	 */
-	void SetEnablePairedEdgeHighlights(bool bEnable);
+	UE_API void SetEnablePairedEdgeHighlights(bool bEnable);
 
 	/**
 	 * Directly set the Color property of the selection line and point set.
 	 */
-	void SetColor(const FColor Color) const;
+	UE_API void SetColor(const FColor Color) const;
 
 	/**
 	 * Directly set the Size/Thickness property of the selection line set.
 	 */
-	void SetLineThickness(const float LineSize) const;
+	UE_API void SetLineThickness(const float LineSize) const;
 
 	/**
 	 * directly set the Size property of the selection point set.
 	 */
-	void SetPointSize(const float PointSize) const;
+	UE_API void SetPointSize(const float PointSize) const;
 
 protected:
 
@@ -136,3 +138,5 @@ protected:
 
 	bool bPairedEdgeHighlightsEnabled = true;
 };
+
+#undef UE_API
diff --git a/Engine/Plugins/Editor/UVEditor/Source/UVEditorTools/Public/TargetInterfaces/UVUnwrapDynamicMesh.h b/Engine/Plugins/Editor/UVEditor/Source/UVEditorTools/Public/TargetInterfaces/UVUnwrapDynamicMesh.h
--- a/Engine/Plugins/Editor/UVEditor/Source/UVEditorTools/Public/TargetInterfaces/UVUnwrapDynamicMesh.h
+++ b/Engine/Plugins/Editor/UVEditor/Source/UVEditorTools/Public/TargetInterfaces/UVUnwrapDynamicMesh.h
@@ -11,8 +11,8 @@
 
 PREDECLARE_USE_GEOMETRY_CLASS(FDynamicMesh3);
 
-UINTERFACE()
-class UVEDITORTOOLS_API UUVUnwrapDynamicMesh : public UInterface
+UINTERFACE(MinimalAPI)
+class UUVUnwrapDynamicMesh : public UInterface
 {
 	GENERATED_BODY()
 };
@@ -22,7 +22,7 @@ class UVEDITORTOOLS_API UUVUnwrapDynamicMesh : public UInterface
  * representing a given layer, and to bake a modified such mesh back to the UV's of
  * the target provided that the triangle id's have all stayed the same.
  */
-class UVEDITORTOOLS_API IUVUnwrapDynamicMesh
+class IUVUnwrapDynamicMesh
 {
 	GENERATED_BODY()
 
@@ -30,4 +30,4 @@ public:
 	virtual int32 GetNumUVLayers() = 0;
 	virtual TSharedPtr<FDynamicMesh3> GetMesh(int32 LayerIndex) = 0;
 	virtual void SaveBackToUVs(const FDynamicMesh3* MeshToSave, int32 LayerIndex) = 0;
-};
\ No newline at end of file
+};
diff --git a/Engine/Plugins/Editor/UVEditor/Source/UVEditorTools/Public/ToolTargets/UVEditorToolMeshInput.h b/Engine/Plugins/Editor/UVEditor/Source/UVEditorTools/Public/ToolTargets/UVEditorToolMeshInput.h
--- a/Engine/Plugins/Editor/UVEditor/Source/UVEditorTools/Public/ToolTargets/UVEditorToolMeshInput.h
+++ b/Engine/Plugins/Editor/UVEditor/Source/UVEditorTools/Public/ToolTargets/UVEditorToolMeshInput.h
@@ -13,6 +13,8 @@
 
 #include "UVEditorToolMeshInput.generated.h"
 
+#define UE_API UVEDITORTOOLS_API
+
 class UMaterialInterface;
 class UMeshOpPreviewWithBackgroundCompute;
 class UMeshElementsVisualizer;
@@ -33,8 +35,8 @@ PREDECLARE_GEOMETRY(class FDynamicMeshChange);
  * 
  * It's arguable whether this should even inherit from UToolTarget.
  */
-UCLASS()
-class UVEDITORTOOLS_API UUVEditorToolMeshInput : public UToolTarget
+UCLASS(MinimalAPI)
+class UUVEditorToolMeshInput : public UToolTarget
 {
 	GENERATED_BODY()
 
@@ -125,7 +127,7 @@ public:
 	TFunction<FVector3d(const FVector2f&)> UVToVertPosition;
 	TFunction<FVector2f(const FVector3d&)> VertPositionToUV;
 
-	bool InitializeMeshes(UToolTarget* Target, 
+	UE_API bool InitializeMeshes(UToolTarget* Target, 
 		TSharedPtr<UE::Geometry::FDynamicMesh3> AppliedCanonicalIn,
 		UMeshOpPreviewWithBackgroundCompute* AppliedPreviewIn,
 		int32 AssetID, int32 UVLayerIndex,
@@ -134,7 +136,7 @@ public:
 		TFunction<FVector3d(const FVector2f&)> UVToVertPositionFuncIn,
 		TFunction<FVector2f(const FVector3d&)> VertPositionToUVFuncIn);
 
-	void Shutdown();
+	UE_API void Shutdown();
 
 	// Notes above the below convenience functions:
 	// 1. Passing nullptr for ChangedVids/ChagnedElementIDs/ChangedConnectivityTids means that all vids/elements/tids
@@ -158,13 +160,13 @@ public:
 
 	// Can be passed in as a ChangedVids/ChangedElements/ChangedConnectivityTids arguments as an equivalent of
 	// passing a pointer to an empty array.
-	static const TArray<int32>* const NONE_CHANGED_ARG;
+	static UE_API const TArray<int32>* const NONE_CHANGED_ARG;
 
 	/**
 	 * Updates UnwrapPreview UV Overlay from UnwrapPreview vert positions. Issues a NotifyDeferredEditCompleted
 	 * for both positions and UVs.
 	 */
-	void UpdateUnwrapPreviewOverlayFromPositions(const TArray<int32>* ChangedVids = nullptr, 
+	UE_API void UpdateUnwrapPreviewOverlayFromPositions(const TArray<int32>* ChangedVids = nullptr, 
 		const TArray<int32>* ChangedConnectivityTids = nullptr, const TArray<int32>* FastRenderUpdateTids = nullptr);
 
 	/**
@@ -174,7 +176,7 @@ public:
 	 * Doesn't broadcast because it is expected that this call is followed by another call that updates the rest of 
 	 * the input object.
 	 */
-	void UpdateUnwrapCanonicalOverlayFromPositions(const TArray<int32>* ChangedVids = nullptr, 
+	UE_API void UpdateUnwrapCanonicalOverlayFromPositions(const TArray<int32>* ChangedVids = nullptr, 
 		const TArray<int32>* ChangedConnectivityTids = nullptr);
 
 	/**
@@ -182,14 +184,14 @@ public:
 	 * a drag, etc, when we only care about updating the visible items.
 	 * Assumes that the overlay in UnwrapPreview is already updated.
 	 */
-	void UpdateAppliedPreviewFromUnwrapPreview(const TArray<int32>* ChangedVids = nullptr, 
+	UE_API void UpdateAppliedPreviewFromUnwrapPreview(const TArray<int32>* ChangedVids = nullptr, 
 		const TArray<int32>* ChangedConnectivityTids = nullptr, const TArray<int32>* FastRenderUpdateTids = nullptr);
 
 	/**
 	 * Updates only the UnwrapPreview from AppliedPreview, without updating the non-preview meshes. Useful for transient
 	 * updates when we only care about updating the visible items.
 	 */
-	void UpdateUnwrapPreviewFromAppliedPreview(
+	UE_API void UpdateUnwrapPreviewFromAppliedPreview(
 		const TArray<int32>* ChangedElementIDs = nullptr, const TArray<int32>* ChangedConnectivityTids = nullptr, 
 		const TArray<int32>* FastRenderUpdateTids = nullptr);
 
@@ -197,13 +199,13 @@ public:
 	 * Updates the non-preview meshes from their preview counterparts. Useful, for instance, after the completion of a drag
 	 * to update the canonical objects.
 	 */
-	void UpdateCanonicalFromPreviews(const TArray<int32>* ChangedVids = nullptr, 
+	UE_API void UpdateCanonicalFromPreviews(const TArray<int32>* ChangedVids = nullptr, 
 		const TArray<int32>* ChangedConnectivityTids = nullptr, bool bBroadcast = true);
 
 	/**
 	 * Updates the preview meshes from their canonical counterparts. Useful mainly as a way to reset the previews.
 	 */
-	void UpdatePreviewsFromCanonical(const TArray<int32>* ChangedVids = nullptr, const TArray<int32>* ChangedConnectivityTids = nullptr,
+	UE_API void UpdatePreviewsFromCanonical(const TArray<int32>* ChangedVids = nullptr, const TArray<int32>* ChangedConnectivityTids = nullptr,
 		const TArray<int32>* FastRenderUpdateTids = nullptr);
 
 	/**
@@ -211,73 +213,75 @@ public:
 	 * Used to update everything when the applied preview was not being updated in tandem with the unwrap preview (otherwise one 
 	 * would use UpdateCanonicalFromPreviews).
 	 */
-	void UpdateAllFromUnwrapPreview(const TArray<int32>* ChangedVids = nullptr, const TArray<int32>* ChangedConnectivityTids = nullptr, 
+	UE_API void UpdateAllFromUnwrapPreview(const TArray<int32>* ChangedVids = nullptr, const TArray<int32>* ChangedConnectivityTids = nullptr, 
 		const TArray<int32>* FastRenderUpdateTids = nullptr, bool bBroadcast = true);
 
 	/**
 	 * Updates the other meshes using the mesh in UnwrapCanonical. Useful when immediately applying an operation that
 	 * operates directly on the unwrap mesh.
 	 */
-	void UpdateAllFromUnwrapCanonical(const TArray<int32>* ChangedVids = nullptr, const TArray<int32>* ChangedConnectivityTids = nullptr, 
+	UE_API void UpdateAllFromUnwrapCanonical(const TArray<int32>* ChangedVids = nullptr, const TArray<int32>* ChangedConnectivityTids = nullptr, 
 		const TArray<int32>* FastRenderUpdateTids = nullptr, bool bBroadcast = true);
 
 	/**
 	 * Updates the other meshes using the mesh in AppliedCanonical. Useful when switching UV layer indices or otherwise
 	 * resetting the collection from the "original" data.
 	 */
-	void UpdateAllFromAppliedCanonical(const TArray<int32>* ChangedVids = nullptr, const TArray<int32>* ChangedConnectivityTids = nullptr, 
+	UE_API void UpdateAllFromAppliedCanonical(const TArray<int32>* ChangedVids = nullptr, const TArray<int32>* ChangedConnectivityTids = nullptr, 
 		const TArray<int32>* FastRenderUpdateTids = nullptr, bool bBroadcast = true);
 
 	/**
 	 * Updates the other meshes using the UV overlay in the live preview.
 	 */
-	void UpdateAllFromAppliedPreview(const TArray<int32>* ChangedElementIDs = nullptr, 
+	UE_API void UpdateAllFromAppliedPreview(const TArray<int32>* ChangedElementIDs = nullptr, 
 		const TArray<int32>* ChangedConnectivityTids = nullptr, const TArray<int32>* FastRenderUpdateTids = nullptr, bool bBroadcast = true);
 
 	/**
 	 * Uses the stored triangles/vertices in the mesh change to update everything from the canonical unwrap. 
 	 * Assumes the change has already been applied to the canonical unwrap.
 	 */
-	void UpdateFromCanonicalUnwrapUsingMeshChange(const UE::Geometry::FDynamicMeshChange& UnwrapCanonicalMeshChange, bool bBroadcast = true);
+	UE_API void UpdateFromCanonicalUnwrapUsingMeshChange(const UE::Geometry::FDynamicMeshChange& UnwrapCanonicalMeshChange, bool bBroadcast = true);
 
 	/**
 	 * Convert a vid in the unwrap mesh to the corresponding vid in the applied mesh (i.e., parent vertex of the
 	 * element corresponding to the unwrap vertex).
 	 */
-	int32 UnwrapVidToAppliedVid(int32 UnwrapVid);
+	UE_API int32 UnwrapVidToAppliedVid(int32 UnwrapVid);
 
 	/**
 	 * Get the unwrap vids corresponding to a given applied vid. These will be multiple if the vertex is
 	 * a seam vertex and therefore has multiple UV elements associated with it.
 	 */
-	void AppliedVidToUnwrapVids(int32 AppliedVid, TArray<int32>& UnwrapVidsOut);
+	UE_API void AppliedVidToUnwrapVids(int32 AppliedVid, TArray<int32>& UnwrapVidsOut);
 
 	/**
 	 * Returns whether or not the underlying source ToolTarget is still valid.
 	 * This is separated from IsValid, which also checks the ToolTarget, in case
 	 * we are interested in the status of just the ToolTarget.
 	*/
-	virtual bool IsToolTargetValid() const;
+	UE_API virtual bool IsToolTargetValid() const;
 
 	/**
 	 * Returns whether or not the canonical and preview meshes are still valid.
 	 * This is separated from IsValid, in case we are interested in the status
 	 * of just the canonical and preview meshes.
 	 */
-	virtual bool AreMeshesValid() const;
+	UE_API virtual bool AreMeshesValid() const;
 
 	// UToolTarget
-	virtual bool IsValid() const override;
+	UE_API virtual bool IsValid() const override;
 
 protected:	
 
 	/** 
 	* Returns the currently configured mesh attribute buffer flags enabled on this target for updating the unwrap previews
 	*/
-	virtual EMeshRenderAttributeFlags GetUnwrapUpdateAttributeFlags() const;
+	UE_API virtual EMeshRenderAttributeFlags GetUnwrapUpdateAttributeFlags() const;
 
 	/**
 	* Returns the currently configured mesh attribute buffer flags enabled on this target for updating the applied previews
 	*/
-	virtual EMeshRenderAttributeFlags GetAppliedUpdateAttributeFlags() const;
-};
\ No newline at end of file
+	UE_API virtual EMeshRenderAttributeFlags GetAppliedUpdateAttributeFlags() const;
+};
+
+#undef UE_API
diff --git a/Engine/Plugins/Editor/UVEditor/Source/UVEditorTools/Public/UVEditorBrushSelectTool.h b/Engine/Plugins/Editor/UVEditor/Source/UVEditorTools/Public/UVEditorBrushSelectTool.h
--- a/Engine/Plugins/Editor/UVEditor/Source/UVEditorTools/Public/UVEditorBrushSelectTool.h
+++ b/Engine/Plugins/Editor/UVEditor/Source/UVEditorTools/Public/UVEditorBrushSelectTool.h
@@ -11,6 +11,8 @@
 
 #include "UVEditorBrushSelectTool.generated.h"
 
+#define UE_API UVEDITORTOOLS_API
+
 class UBrushStampIndicator;
 class ULocalMouseHoverBehavior;
 class ULocalClickDragInputBehavior;
@@ -53,8 +55,8 @@ public:
  * The brush select tool allows for brush selection of triangles on the unwrap or the live
  *  preview. 
  */
-UCLASS()
-class UVEDITORTOOLS_API UUVEditorBrushSelectTool : public UInteractiveTool
+UCLASS(MinimalAPI)
+class UUVEditorBrushSelectTool : public UInteractiveTool
 	, public IInteractiveToolNestedAcceptCancelAPI
 	, public IUVToolSupportsSelection
 	, public IUVEditorGenericBuildableTool
@@ -64,18 +66,18 @@ class UVEDITORTOOLS_API UUVEditorBrushSelectTool : public UInteractiveTool
 public:
 	
 	// UInteractiveTool
-	virtual void Setup() override;
-	virtual void Shutdown(EToolShutdownType ShutdownType) override;
-	virtual void Render(IToolsContextRenderAPI* RenderAPI) override;
-	virtual void OnTick(float DeltaTime) override;
+	UE_API virtual void Setup() override;
+	UE_API virtual void Shutdown(EToolShutdownType ShutdownType) override;
+	UE_API virtual void Render(IToolsContextRenderAPI* RenderAPI) override;
+	UE_API virtual void OnTick(float DeltaTime) override;
 	virtual bool HasCancel() const override { return false; }
 	virtual bool HasAccept() const override { return false; }
-	virtual void RegisterActions(FInteractiveToolActionSet& ActionSet) override;
+	UE_API virtual void RegisterActions(FInteractiveToolActionSet& ActionSet) override;
 
 	// IInteractiveToolNestedAcceptCancelAPI
 	virtual bool SupportsNestedCancelCommand() override { return true; }
-	virtual bool CanCurrentlyNestedCancel() override;
-	virtual bool ExecuteNestedCancelCommand() override;
+	UE_API virtual bool CanCurrentlyNestedCancel() override;
+	UE_API virtual bool ExecuteNestedCancelCommand() override;
 
 	// IUVEditorGenericBuildableTool
 	virtual void SetTargets(const TArray<TObjectPtr<UUVEditorToolMeshInput>>& TargetsIn)
@@ -121,10 +123,10 @@ private:
 	TArray<FVector2d> PendingUnwrapHits;
 
 	// Updates internal structures in addition to the actual selection api
-	void ClearSelections(bool bBroadcastAndEmit = false);
+	UE_API void ClearSelections(bool bBroadcastAndEmit = false);
 
-	void ProcessPendingUnwrapHits();
-	void ProcessPendingLivePreviewHits();
+	UE_API void ProcessPendingUnwrapHits();
+	UE_API void ProcessPendingLivePreviewHits();
 
 	bool bHavePendingUnwrapHit = false;
 	bool bHavePendingLivePreviewHit = false;
@@ -154,8 +156,10 @@ private:
 	bool bHoveringLivePreview = false;
 	bool bDraggingUnwrap = false;
 	bool bDraggingLivePreview = false;
-	void UpdateViewportStateFromHoverOrDragEvent(bool bFromUnwrap, bool bIsEndEvent, bool bDragging);
+	UE_API void UpdateViewportStateFromHoverOrDragEvent(bool bFromUnwrap, bool bIsEndEvent, bool bDragging);
 	
-	void IncreaseBrushRadiusAction();
-	void DecreaseBrushRadiusAction();
-};
\ No newline at end of file
+	UE_API void IncreaseBrushRadiusAction();
+	UE_API void DecreaseBrushRadiusAction();
+};
+
+#undef UE_API
diff --git a/Engine/Plugins/Editor/UVEditor/Source/UVEditorTools/Public/UVEditorLayerEditTool.h b/Engine/Plugins/Editor/UVEditor/Source/UVEditorTools/Public/UVEditorLayerEditTool.h
--- a/Engine/Plugins/Editor/UVEditor/Source/UVEditorTools/Public/UVEditorLayerEditTool.h
+++ b/Engine/Plugins/Editor/UVEditor/Source/UVEditorTools/Public/UVEditorLayerEditTool.h
@@ -10,18 +10,20 @@
 
 #include "UVEditorLayerEditTool.generated.h"
 
+#define UE_API UVEDITORTOOLS_API
+
 class UUVEditorToolMeshInput;
 class UUVEditorChannelEditTool;
 class UUVToolEmitChangeAPI;
 
-UCLASS()
-class UVEDITORTOOLS_API UUVEditorChannelEditToolBuilder : public UInteractiveToolBuilder
+UCLASS(MinimalAPI)
+class UUVEditorChannelEditToolBuilder : public UInteractiveToolBuilder
 {
 	GENERATED_BODY()
 
 public:
-	virtual bool CanBuildTool(const FToolBuilderState& SceneState) const override;
-	virtual UInteractiveTool* BuildTool(const FToolBuilderState& SceneState) const override;
+	UE_API virtual bool CanBuildTool(const FToolBuilderState& SceneState) const override;
+	UE_API virtual UInteractiveTool* BuildTool(const FToolBuilderState& SceneState) const override;
 
 	// This is a pointer so that it can be updated under the builder without
 	// having to set it in the mode after initializing targets.
@@ -39,8 +41,8 @@ enum class EChannelEditToolAction
 };
 
 
-UCLASS()
-class UVEDITORTOOLS_API UUVEditorChannelEditSettings : public UInteractiveToolPropertySet
+UCLASS(MinimalAPI)
+class UUVEditorChannelEditSettings : public UInteractiveToolPropertySet
 {
 	GENERATED_BODY()
 public:
@@ -49,8 +51,8 @@ public:
 	EChannelEditToolAction Action = EChannelEditToolAction::Add;
 };
 
-UCLASS()
-class UVEDITORTOOLS_API UUVEditorChannelEditTargetProperties : public UInteractiveToolPropertySet
+UCLASS(MinimalAPI)
+class UUVEditorChannelEditTargetProperties : public UInteractiveToolPropertySet
 {
 	GENERATED_BODY()
 public:
@@ -59,7 +61,7 @@ public:
 	FString Asset;
 
 	UFUNCTION()
-	const TArray<FString>& GetAssetNames();
+	UE_API const TArray<FString>& GetAssetNames();
 
 	UPROPERTY(EditAnywhere, Category = "UVChannels", meta = (DisplayName = "Target UV Channel", GetOptions = GetUVChannelNames, EditCondition = "bActionNeedsTarget", EditConditionHides = true, HideEditConditionToggle = true))
 	FString TargetChannel;
@@ -68,7 +70,7 @@ public:
 	FString ReferenceChannel;
 
 	UFUNCTION()
-	const TArray<FString>& GetUVChannelNames();
+	UE_API const TArray<FString>& GetUVChannelNames();
 
 	TArray<FString> UVChannelNames;
 	TArray<FString> UVAssetNames;
@@ -86,66 +88,66 @@ public:
 	bool bActionNeedsTarget = true;
 
 public:
-	void Initialize(
+	UE_API void Initialize(
 		const TArray<TObjectPtr<UUVEditorToolMeshInput>>& TargetsIn,
 		bool bInitializeSelection);
 
-	void SetUsageFlags(EChannelEditToolAction Action);
+	UE_API void SetUsageFlags(EChannelEditToolAction Action);
 
 	/**
 	 * Verify that the UV asset selection is valid
 	 * @param bUpdateIfInvalid if selection is not valid, reset UVAsset to Asset0 or empty if no assets exist
 	 * @return true if selection in UVAsset is an entry in UVAssetNames.
 	 */
-	bool ValidateUVAssetSelection(bool bUpdateIfInvalid);
+	UE_API bool ValidateUVAssetSelection(bool bUpdateIfInvalid);
 
 	/**
 	 * Verify that the UV channel selection is valid
 	 * @param bUpdateIfInvalid if selection is not valid, UVChannel to UV0 or empty if no UV channels exist
 	 * @return true if selection in UVChannel is an entry in UVChannelNamesList.
 	 */
-	bool ValidateUVChannelSelection(bool bUpdateIfInvalid);
+	UE_API bool ValidateUVChannelSelection(bool bUpdateIfInvalid);
 
 
 	/**
 	 * @return selected UV asset ID, or -1 if invalid selection
 	 */
-	int32 GetSelectedAssetID();
+	UE_API int32 GetSelectedAssetID();
 
 	/**
 	 * @param bForceToZeroOnFailure if true, then instead of returning -1 we return 0 so calling code can fallback to default UV paths
 	 * @param bUseReference if true, get the selected reference channel index, otherwise return the target channel's index.
 	 * @return selected UV channel index, or -1 if invalid selection, or 0 if invalid selection and bool bForceToZeroOnFailure = true
 	 */
-	int32 GetSelectedChannelIndex(bool bForceToZeroOnFailure = false, bool bUseReference = false);
+	UE_API int32 GetSelectedChannelIndex(bool bForceToZeroOnFailure = false, bool bUseReference = false);
 };
 
-UCLASS()
-class UVEDITORTOOLS_API UUVEditorChannelEditAddProperties : public UInteractiveToolPropertySet
+UCLASS(MinimalAPI)
+class UUVEditorChannelEditAddProperties : public UInteractiveToolPropertySet
 {
 	GENERATED_BODY()
 public:
 	/* Placeholder for future per action settings */
 };
 
-UCLASS()
-class UVEDITORTOOLS_API UUVEditorChannelEditCopyProperties : public UInteractiveToolPropertySet
+UCLASS(MinimalAPI)
+class UUVEditorChannelEditCopyProperties : public UInteractiveToolPropertySet
 {
 	GENERATED_BODY()
 public:
 	/* Placeholder for future per action settings */
 };
 
-UCLASS()
-class UVEDITORTOOLS_API UUVEditorChannelEditDeleteProperties : public UInteractiveToolPropertySet
+UCLASS(MinimalAPI)
+class UUVEditorChannelEditDeleteProperties : public UInteractiveToolPropertySet
 {
 	GENERATED_BODY()
 public:
 	/* Placeholder for future per action settings */
 };
 
-UCLASS()
-class UVEDITORTOOLS_API UUVEditorChannelEditToolActionPropertySet : public UInteractiveToolPropertySet
+UCLASS(MinimalAPI)
+class UUVEditorChannelEditToolActionPropertySet : public UInteractiveToolPropertySet
 {
 	GENERATED_BODY()
 
@@ -153,17 +155,17 @@ public:
 	TWeakObjectPtr<UUVEditorChannelEditTool> ParentTool;
 
 	void Initialize(UUVEditorChannelEditTool* ParentToolIn) { ParentTool = ParentToolIn; }
-	void PostAction(EChannelEditToolAction Action);
+	UE_API void PostAction(EChannelEditToolAction Action);
 	
 	UFUNCTION(CallInEditor, Category = Actions)
-	void Apply();
+	UE_API void Apply();
 };
 
 /**
  *
  */
-UCLASS()
-class UVEDITORTOOLS_API UUVEditorChannelEditTool : public UInteractiveTool
+UCLASS(MinimalAPI)
+class UUVEditorChannelEditTool : public UInteractiveTool
 {
 	GENERATED_BODY()
 
@@ -177,28 +179,28 @@ public:
 		Targets = TargetsIn;
 	}
 
-	virtual void Setup() override;
-	virtual void Shutdown(EToolShutdownType ShutdownType) override;
+	UE_API virtual void Setup() override;
+	UE_API virtual void Shutdown(EToolShutdownType ShutdownType) override;
 
-	virtual void OnTick(float DeltaTime) override;
+	UE_API virtual void OnTick(float DeltaTime) override;
 
 	virtual bool HasCancel() const override { return false; }
 	virtual bool HasAccept() const override { return false; }
 	virtual bool CanAccept() const override { return false; }
 
-	virtual void OnPropertyModified(UObject* PropertySet, FProperty* Property) override;
+	UE_API virtual void OnPropertyModified(UObject* PropertySet, FProperty* Property) override;
 
-	void RequestAction(EChannelEditToolAction ActionType);
-	EChannelEditToolAction ActiveAction() const;
+	UE_API void RequestAction(EChannelEditToolAction ActionType);
+	UE_API EChannelEditToolAction ActiveAction() const;
 
 
 protected:
-	void ApplyVisbleChannelChange();
-	void UpdateChannelSelectionProperties(int32 ChangingAsset);
+	UE_API void ApplyVisbleChannelChange();
+	UE_API void UpdateChannelSelectionProperties(int32 ChangingAsset);
 
-	void AddChannel();
-	void CopyChannel();
-	void DeleteChannel();
+	UE_API void AddChannel();
+	UE_API void CopyChannel();
+	UE_API void DeleteChannel();
 	int32 ActiveAsset;
 	int32 ActiveChannel;
 	int32 ReferenceChannel;
@@ -254,5 +256,7 @@ protected:
 	TArray<FActionHistoryItem> AnalyticsActionHistory;
 	UE::Geometry::UVEditorAnalytics::FTargetAnalytics InputTargetAnalytics;
 	FDateTime ToolStartTimeAnalytics;
-	void RecordAnalytics();
+	UE_API void RecordAnalytics();
 };
+
+#undef UE_API
diff --git a/Engine/Plugins/Editor/UVEditor/Source/UVEditorTools/Public/UVEditorLayoutTool.h b/Engine/Plugins/Editor/UVEditor/Source/UVEditorTools/Public/UVEditorLayoutTool.h
--- a/Engine/Plugins/Editor/UVEditor/Source/UVEditorTools/Public/UVEditorLayoutTool.h
+++ b/Engine/Plugins/Editor/UVEditor/Source/UVEditorTools/Public/UVEditorLayoutTool.h
@@ -11,19 +11,21 @@
 
 #include "UVEditorLayoutTool.generated.h"
 
+#define UE_API UVEDITORTOOLS_API
+
 class UUVEditorToolMeshInput;
 class UUVLayoutProperties;
 class UUVLayoutOperatorFactory;
 class UUVTool2DViewportAPI;
 
-UCLASS()
-class UVEDITORTOOLS_API UUVEditorLayoutToolBuilder : public UInteractiveToolBuilder
+UCLASS(MinimalAPI)
+class UUVEditorLayoutToolBuilder : public UInteractiveToolBuilder
 {
 	GENERATED_BODY()
 
 public:
-	virtual bool CanBuildTool(const FToolBuilderState& SceneState) const override;
-	virtual UInteractiveTool* BuildTool(const FToolBuilderState& SceneState) const override;
+	UE_API virtual bool CanBuildTool(const FToolBuilderState& SceneState) const override;
+	UE_API virtual UInteractiveTool* BuildTool(const FToolBuilderState& SceneState) const override;
 
 	// This is a pointer so that it can be updated under the builder without
 	// having to set it in the mode after initializing targets.
@@ -33,8 +35,8 @@ public:
 /**
  * 
  */
-UCLASS()
-class UVEDITORTOOLS_API UUVEditorLayoutTool : public UInteractiveTool, public IUVToolSupportsSelection
+UCLASS(MinimalAPI)
+class UUVEditorLayoutTool : public UInteractiveTool, public IUVToolSupportsSelection
 {
 	GENERATED_BODY()
 
@@ -48,16 +50,16 @@ public:
 		Targets = TargetsIn;
 	}
 
-	virtual void Setup() override;
-	virtual void Shutdown(EToolShutdownType ShutdownType) override;
+	UE_API virtual void Setup() override;
+	UE_API virtual void Shutdown(EToolShutdownType ShutdownType) override;
 
-	virtual void OnTick(float DeltaTime) override;
+	UE_API virtual void OnTick(float DeltaTime) override;
 
 	virtual bool HasCancel() const override { return true; }
 	virtual bool HasAccept() const override { return true; }
-	virtual bool CanAccept() const override;
+	UE_API virtual bool CanAccept() const override;
 
-	virtual void OnPropertyModified(UObject* PropertySet, FProperty* Property) override;
+	UE_API virtual void OnPropertyModified(UObject* PropertySet, FProperty* Property) override;
 protected:
 
 	UPROPERTY()
@@ -82,5 +84,7 @@ protected:
 
 	UE::Geometry::UVEditorAnalytics::FTargetAnalytics InputTargetAnalytics;
 	FDateTime ToolStartTimeAnalytics;
-	void RecordAnalytics();
+	UE_API void RecordAnalytics();
 };
+
+#undef UE_API
