Updated headers using UnrealCodeFixup to make sure dllstorage is on methods/staticvars instead of types. Part 2/n

diff --git a/Engine/Plugins/Editor/UVEditor/Source/UVEditorTools/Public/UVEditorMechanicAdapterTool.h b/Engine/Plugins/Editor/UVEditor/Source/UVEditorTools/Public/UVEditorMechanicAdapterTool.h
--- a/Engine/Plugins/Editor/UVEditor/Source/UVEditorTools/Public/UVEditorMechanicAdapterTool.h
+++ b/Engine/Plugins/Editor/UVEditor/Source/UVEditorTools/Public/UVEditorMechanicAdapterTool.h
@@ -16,8 +16,8 @@ class UInteractiveToolManager;
  * If necessary, the behaviors can be extracted via GetInputBehaviors() to be placed into 
  * an input router.
  */
-UCLASS()
-class UVEDITORTOOLS_API UUVEditorMechanicAdapterTool : public UInteractiveTool
+UCLASS(MinimalAPI)
+class UUVEditorMechanicAdapterTool : public UInteractiveTool
 {
 	GENERATED_BODY()
 
@@ -25,4 +25,4 @@ public:
 	UInteractiveToolManager* ToolManager = nullptr;
 
 	virtual UInteractiveToolManager* GetToolManager() const override { return ToolManager; }
-};
\ No newline at end of file
+};
diff --git a/Engine/Plugins/Editor/UVEditor/Source/UVEditorTools/Public/UVEditorRecomputeUVsTool.h b/Engine/Plugins/Editor/UVEditor/Source/UVEditorTools/Public/UVEditorRecomputeUVsTool.h
--- a/Engine/Plugins/Editor/UVEditor/Source/UVEditorTools/Public/UVEditorRecomputeUVsTool.h
+++ b/Engine/Plugins/Editor/UVEditor/Source/UVEditorTools/Public/UVEditorRecomputeUVsTool.h
@@ -13,6 +13,8 @@
 
 #include "UVEditorRecomputeUVsTool.generated.h"
 
+#define UE_API UVEDITORTOOLS_API
+
 
 // Forward declarations
 class UDynamicMeshComponent;
@@ -25,14 +27,14 @@ class URecomputeUVsOpFactory;
 /**
  *
  */
-UCLASS()
-class UVEDITORTOOLS_API UUVEditorRecomputeUVsToolBuilder : public UInteractiveToolBuilder
+UCLASS(MinimalAPI)
+class UUVEditorRecomputeUVsToolBuilder : public UInteractiveToolBuilder
 {
 	GENERATED_BODY()
 
 public:
-	virtual bool CanBuildTool(const FToolBuilderState& SceneState) const override;
-	virtual UInteractiveTool* BuildTool(const FToolBuilderState& SceneState) const override;
+	UE_API virtual bool CanBuildTool(const FToolBuilderState& SceneState) const override;
+	UE_API virtual UInteractiveTool* BuildTool(const FToolBuilderState& SceneState) const override;
 
 	// This is a pointer so that it can be updated under the builder without
 	// having to set it in the mode after initializing targets.
@@ -43,24 +45,24 @@ public:
 /**
  * UUVEditorRecomputeUVsTool Recomputes UVs based on existing segmentations of the mesh
  */
-UCLASS()
-class UVEDITORTOOLS_API UUVEditorRecomputeUVsTool : public UInteractiveTool, public IUVToolSupportsSelection
+UCLASS(MinimalAPI)
+class UUVEditorRecomputeUVsTool : public UInteractiveTool, public IUVToolSupportsSelection
 {
 	GENERATED_BODY()
 
 public:
-	virtual void Setup() override;
-	virtual void Shutdown(EToolShutdownType ShutdownType) override;
+	UE_API virtual void Setup() override;
+	UE_API virtual void Shutdown(EToolShutdownType ShutdownType) override;
 	
-	virtual void OnTick(float DeltaTime) override;
+	UE_API virtual void OnTick(float DeltaTime) override;
 
-	virtual void Render(IToolsContextRenderAPI* RenderAPI) override;
+	UE_API virtual void Render(IToolsContextRenderAPI* RenderAPI) override;
 
 	virtual bool HasCancel() const override { return true; }
 	virtual bool HasAccept() const override { return true; }
-	virtual bool CanAccept() const override;
+	UE_API virtual bool CanAccept() const override;
 
-	virtual void OnPropertyModified(UObject* PropertySet, FProperty* Property) override;
+	UE_API virtual void OnPropertyModified(UObject* PropertySet, FProperty* Property) override;
 
 
 	/**
@@ -89,8 +91,8 @@ protected:
 
 
 	TSharedPtr<UE::Geometry::FPolygroupSet, ESPMode::ThreadSafe> ActiveGroupSet;
-	void OnSelectedGroupLayerChanged();
-	void UpdateActiveGroupLayer(bool bUpdateFactories = true);
+	UE_API void OnSelectedGroupLayerChanged();
+	UE_API void UpdateActiveGroupLayer(bool bUpdateFactories = true);
 
 	//
 	// Analytics
@@ -98,5 +100,7 @@ protected:
 	
 	UE::Geometry::UVEditorAnalytics::FTargetAnalytics InputTargetAnalytics;
 	FDateTime ToolStartTimeAnalytics;
-	void RecordAnalytics();
+	UE_API void RecordAnalytics();
 };
+
+#undef UE_API
diff --git a/Engine/Plugins/Editor/UVEditor/Source/UVEditorTools/Public/UVEditorSeamTool.h b/Engine/Plugins/Editor/UVEditor/Source/UVEditorTools/Public/UVEditorSeamTool.h
--- a/Engine/Plugins/Editor/UVEditor/Source/UVEditorTools/Public/UVEditorSeamTool.h
+++ b/Engine/Plugins/Editor/UVEditor/Source/UVEditorTools/Public/UVEditorSeamTool.h
@@ -14,6 +14,8 @@
 
 #include "UVEditorSeamTool.generated.h"
 
+#define UE_API UVEDITORTOOLS_API
+
 //class ULocalSingleClickInputBehavior;
 class UInputRouter;
 class ULocalInputBehaviorSource;
@@ -33,8 +35,8 @@ enum class EUVEditorSeamMode : uint8
 	Join = 1
 };
 
-UCLASS()
-class UVEDITORTOOLS_API UUVEditorSeamToolProperties : public UInteractiveToolPropertySet
+UCLASS(MinimalAPI)
+class UUVEditorSeamToolProperties : public UInteractiveToolPropertySet
 {
 	GENERATED_BODY()
 
@@ -54,22 +56,22 @@ public:
 	double PathSimilarityWeight = 200;
 };
 
-UCLASS()
-class UVEDITORTOOLS_API UUVEditorSeamToolBuilder : public UInteractiveToolBuilder
+UCLASS(MinimalAPI)
+class UUVEditorSeamToolBuilder : public UInteractiveToolBuilder
 {
 	GENERATED_BODY()
 public:
 
-	virtual bool CanBuildTool(const FToolBuilderState& SceneState) const override;
-	virtual UInteractiveTool* BuildTool(const FToolBuilderState& SceneState) const override;
+	UE_API virtual bool CanBuildTool(const FToolBuilderState& SceneState) const override;
+	UE_API virtual UInteractiveTool* BuildTool(const FToolBuilderState& SceneState) const override;
 
 	// These are pointers so that they can be updated under the builder without
 	// having to reset them after things are reinitialized.
 	const TArray<TObjectPtr<UUVEditorToolMeshInput>>* Targets = nullptr;
 };
 
-UCLASS()
-class UVEDITORTOOLS_API UUVEditorSeamTool : public UInteractiveTool,
+UCLASS(MinimalAPI)
+class UUVEditorSeamTool : public UInteractiveTool,
 	public IInteractiveToolNestedAcceptCancelAPI
 {
 	GENERATED_BODY()
@@ -77,50 +79,50 @@ class UVEDITORTOOLS_API UUVEditorSeamTool : public UInteractiveTool,
 public:
 	using FDynamicMeshAABBTree3 = UE::Geometry::FDynamicMeshAABBTree3;
 
-	virtual void SetTargets(const TArray<TObjectPtr<UUVEditorToolMeshInput>>& TargetsIn);
+	UE_API virtual void SetTargets(const TArray<TObjectPtr<UUVEditorToolMeshInput>>& TargetsIn);
 
 	// For use by undo/redo
-	void EditLockedPath(
+	UE_API void EditLockedPath(
 		TUniqueFunction<void(TArray<int32>& LockedPathInOut)> EditFunction, int32 MeshIndex);
 
 	// UInteractiveTool
-	virtual void Setup() override;
-	virtual void Shutdown(EToolShutdownType ShutdownType) override;
-	virtual void OnTick(float DeltaTime) override;
+	UE_API virtual void Setup() override;
+	UE_API virtual void Shutdown(EToolShutdownType ShutdownType) override;
+	UE_API virtual void OnTick(float DeltaTime) override;
 	virtual bool HasCancel() const override { return false; }
 	virtual bool HasAccept() const override { return false; }
 
 	// IInteractiveToolNestedAcceptCancelAPI
 	virtual bool SupportsNestedCancelCommand() override { return true; }
-	virtual bool CanCurrentlyNestedCancel() override;
-	virtual bool ExecuteNestedCancelCommand() override;
+	UE_API virtual bool CanCurrentlyNestedCancel() override;
+	UE_API virtual bool ExecuteNestedCancelCommand() override;
 	virtual bool SupportsNestedAcceptCommand() override { return true; }
-	virtual bool CanCurrentlyNestedAccept() override;
-	virtual bool ExecuteNestedAcceptCommand() override;
+	UE_API virtual bool CanCurrentlyNestedAccept() override;
+	UE_API virtual bool ExecuteNestedAcceptCommand() override;
 protected:
-	void ReconstructExistingSeamsVisualization();
-	void ReconstructLockedPathVisualization();
+	UE_API void ReconstructExistingSeamsVisualization();
+	UE_API void ReconstructLockedPathVisualization();
 
-	int32 Get2DHitVertex(const FRay& WorldRayIn, int32* IndexOf2DSpatialOut = nullptr);
-	int32 Get3DHitVertex(const FRay& WorldRayIn, int32* IndexOf3DSpatialOut = nullptr);
+	UE_API int32 Get2DHitVertex(const FRay& WorldRayIn, int32* IndexOf2DSpatialOut = nullptr);
+	UE_API int32 Get3DHitVertex(const FRay& WorldRayIn, int32* IndexOf3DSpatialOut = nullptr);
 
-	void OnMeshVertexClicked(int32 Vid, int32 IndexOfMesh, bool bVidIsFromUnwrap);
-	void OnMeshVertexHovered(int32 Vid, int32 IndexOfMesh, bool bVidIsFromUnwrap);
-	void UpdateHover();
+	UE_API void OnMeshVertexClicked(int32 Vid, int32 IndexOfMesh, bool bVidIsFromUnwrap);
+	UE_API void OnMeshVertexHovered(int32 Vid, int32 IndexOfMesh, bool bVidIsFromUnwrap);
+	UE_API void UpdateHover();
 	/** @param bClearHoverInfo If true, also clears HoverVid, etc in addition to just clearing display. */
-	void ClearHover(bool bClearHoverInfo = true);
-	void ResetPreviewColors();
-	void ApplyClick();
+	UE_API void ClearHover(bool bClearHoverInfo = true);
+	UE_API void ResetPreviewColors();
+	UE_API void ApplyClick();
 	/** Clears LockedAppliedVids, but builds seam off of AppliedVidsIn */
-	void ApplySeam(const TArray<int32>& AppliedVidsIn);
-	void ClearLockedPath(bool bEmitChange = true);
+	UE_API void ApplySeam(const TArray<int32>& AppliedVidsIn);
+	UE_API void ClearLockedPath(bool bEmitChange = true);
 
-	void UpdateToolMessage();
+	UE_API void UpdateToolMessage();
 
 	int32 TemporaryModeToggleModifierID = 1;
 	bool bModeIsTemporarilyToggled;
-	bool IsInJoinMode() const;
-	void OnSeamModeChanged();
+	UE_API bool IsInJoinMode() const;
+	UE_API void OnSeamModeChanged();
 
 	enum class EState
 	{
@@ -181,8 +183,8 @@ protected:
 	TArray<int32> LockedPath;
 	TArray<int32> PreviewPath;
 
-	FColor GetLockedPathColor() const;
-	FColor GetExtendPathColor() const;
+	UE_API FColor GetLockedPathColor() const;
+	UE_API FColor GetExtendPathColor() const;
 
 	// When true, the entire path is changed to the "completion" color to show
 	// that the next click will complete the path.
@@ -194,12 +196,14 @@ protected:
 	
 	UE::Geometry::UVEditorAnalytics::FTargetAnalytics InputTargetAnalytics;
 	FDateTime ToolStartTimeAnalytics;
-	void RecordAnalytics();
+	UE_API void RecordAnalytics();
 
 private:
-	void GetVidPath(const UUVEditorToolMeshInput& Target, bool bUnwrapMeshSource,
+	UE_API void GetVidPath(const UUVEditorToolMeshInput& Target, bool bUnwrapMeshSource,
 		const TArray<int32>& StartVids, int32 EndVid, TArray<int32>& VidPathOut);
 
 	TArray<double> UnwrapMaxDims;
 	TArray<double> AppliedMeshMaxDims;
 };
+
+#undef UE_API
diff --git a/Engine/Plugins/Editor/UVEditor/Source/UVEditorTools/Public/UVEditorStyle.h b/Engine/Plugins/Editor/UVEditor/Source/UVEditorTools/Public/UVEditorStyle.h
--- a/Engine/Plugins/Editor/UVEditor/Source/UVEditorTools/Public/UVEditorStyle.h
+++ b/Engine/Plugins/Editor/UVEditor/Source/UVEditorTools/Public/UVEditorStyle.h
@@ -5,20 +5,24 @@
 #include "Styling/SlateStyle.h"
 #include "Styling/SlateStyleRegistry.h"
 
+#define UE_API UVEDITORTOOLS_API
+
 /**
  * Slate style set for UV Editor
  */
-class UVEDITORTOOLS_API FUVEditorStyle
+class FUVEditorStyle
 	: public FSlateStyleSet
 {
 public:
-	static FName StyleName;
+	static UE_API FName StyleName;
 
 	/** Access the singleton instance for this style set */
-	static FUVEditorStyle& Get();
+	static UE_API FUVEditorStyle& Get();
 
 private:
 
-	FUVEditorStyle();
-	~FUVEditorStyle();
-};
\ No newline at end of file
+	UE_API FUVEditorStyle();
+	UE_API ~FUVEditorStyle();
+};
+
+#undef UE_API
diff --git a/Engine/Plugins/Editor/UVEditor/Source/UVEditorTools/Public/UVEditorTexelDensityTool.h b/Engine/Plugins/Editor/UVEditor/Source/UVEditorTools/Public/UVEditorTexelDensityTool.h
--- a/Engine/Plugins/Editor/UVEditor/Source/UVEditorTools/Public/UVEditorTexelDensityTool.h
+++ b/Engine/Plugins/Editor/UVEditor/Source/UVEditorTools/Public/UVEditorTexelDensityTool.h
@@ -15,6 +15,8 @@
 
 #include "UVEditorTexelDensityTool.generated.h"
 
+#define UE_API UVEDITORTOOLS_API
+
 
 // Forward declarations
 class UDynamicMeshComponent;
@@ -31,14 +33,14 @@ class UPreviewGeometry;
 /**
  *
  */
-UCLASS()
-class UVEDITORTOOLS_API UUVEditorTexelDensityToolBuilder : public UInteractiveToolBuilder
+UCLASS(MinimalAPI)
+class UUVEditorTexelDensityToolBuilder : public UInteractiveToolBuilder
 {
 	GENERATED_BODY()
 
 public:
-	virtual bool CanBuildTool(const FToolBuilderState& SceneState) const override;
-	virtual UInteractiveTool* BuildTool(const FToolBuilderState& SceneState) const override;
+	UE_API virtual bool CanBuildTool(const FToolBuilderState& SceneState) const override;
+	UE_API virtual UInteractiveTool* BuildTool(const FToolBuilderState& SceneState) const override;
 
 	// This is a pointer so that it can be updated under the builder without
 	// having to set it in the mode after initializing targets.
@@ -56,8 +58,8 @@ enum class ETexelDensityToolAction
 };
 
 
-UCLASS()
-class UVEDITORTOOLS_API UUVEditorTexelDensityActionSettings : public UInteractiveToolPropertySet
+UCLASS(MinimalAPI)
+class UUVEditorTexelDensityActionSettings : public UInteractiveToolPropertySet
 {
 	GENERATED_BODY()
 public:
@@ -65,14 +67,14 @@ public:
 	TWeakObjectPtr<UUVEditorTexelDensityTool> ParentTool;
 
 	void Initialize(UUVEditorTexelDensityTool* ParentToolIn) { ParentTool = ParentToolIn; }
-	void PostAction(ETexelDensityToolAction Action);
+	UE_API void PostAction(ETexelDensityToolAction Action);
 
 	UFUNCTION(CallInEditor, Category = Actions)
-	void SampleTexelDensity();
+	UE_API void SampleTexelDensity();
 };
 
-UCLASS()
-class UVEDITORTOOLS_API UUVEditorTexelDensityToolSettings : public UUVEditorTexelDensitySettings
+UCLASS(MinimalAPI)
+class UUVEditorTexelDensityToolSettings : public UUVEditorTexelDensitySettings
 {
 	GENERATED_BODY()
 public:
@@ -80,30 +82,30 @@ public:
 	TWeakObjectPtr<UUVEditorTexelDensityTool> ParentTool;
 	void Initialize(UUVEditorTexelDensityTool* ParentToolIn) { ParentTool = ParentToolIn; }
 
-	virtual bool InSamplingMode() const override;
+	UE_API virtual bool InSamplingMode() const override;
 };
 
 /**
  * UUVEditorRecomputeUVsTool Recomputes UVs based on existing segmentations of the mesh
  */
-UCLASS()
-class UVEDITORTOOLS_API UUVEditorTexelDensityTool : public UInteractiveTool, public IUVToolSupportsSelection
+UCLASS(MinimalAPI)
+class UUVEditorTexelDensityTool : public UInteractiveTool, public IUVToolSupportsSelection
 {
 	GENERATED_BODY()
 
 public:
-	virtual void Setup() override;
-	virtual void Shutdown(EToolShutdownType ShutdownType) override;
+	UE_API virtual void Setup() override;
+	UE_API virtual void Shutdown(EToolShutdownType ShutdownType) override;
 
-	virtual void OnTick(float DeltaTime) override;
+	UE_API virtual void OnTick(float DeltaTime) override;
 
-	virtual void Render(IToolsContextRenderAPI* RenderAPI) override;
+	UE_API virtual void Render(IToolsContextRenderAPI* RenderAPI) override;
 
 	virtual bool HasCancel() const override { return true; }
 	virtual bool HasAccept() const override { return true; }
-	virtual bool CanAccept() const override;
+	UE_API virtual bool CanAccept() const override;
 
-	virtual void OnPropertyModified(UObject* PropertySet, FProperty* Property) override;
+	UE_API virtual void OnPropertyModified(UObject* PropertySet, FProperty* Property) override;
 
 
 	/**
@@ -115,8 +117,8 @@ public:
 	}
 
 
-	void RequestAction(ETexelDensityToolAction ActionType);
-	ETexelDensityToolAction ActiveAction() const;
+	UE_API void RequestAction(ETexelDensityToolAction ActionType);
+	UE_API ETexelDensityToolAction ActiveAction() const;
 
 private:
 	UPROPERTY()
@@ -147,7 +149,7 @@ private:
 	TArray<TObjectPtr<UUVTexelDensityOperatorFactory>> Factories;
 
 	int32 RemainingTargetsToProcess;
-	void PerformBackgroundScalingTask();
+	UE_API void PerformBackgroundScalingTask();
 
 	UPROPERTY()
 	TObjectPtr<UInputBehaviorSet> LivePreviewBehaviorSet = nullptr;
@@ -172,20 +174,20 @@ private:
 	TObjectPtr<UMaterialInstanceDynamic> TriangleSetMaterial = nullptr;
 
 	TSharedPtr<UE::Geometry::FPolygroupSet, ESPMode::ThreadSafe> ActiveGroupSet;
-	void OnSelectedGroupLayerChanged();
-	void UpdateActiveGroupLayer(bool bUpdateFactories = true);
+	UE_API void OnSelectedGroupLayerChanged();
+	UE_API void UpdateActiveGroupLayer(bool bUpdateFactories = true);
 
-	int32 Get2DHitTriangle(const FRay& WorldRayIn, int32* IndexOf2DSpatialOut = nullptr);
-	int32 Get3DHitTriangle(const FRay& WorldRayIn, int32* IndexOf2DSpatialOut = nullptr);
+	UE_API int32 Get2DHitTriangle(const FRay& WorldRayIn, int32* IndexOf2DSpatialOut = nullptr);
+	UE_API int32 Get3DHitTriangle(const FRay& WorldRayIn, int32* IndexOf2DSpatialOut = nullptr);
 
-	void OnMeshTriangleClicked(int32 Tid, int32 IndexOfMesh, bool bTidIsFromUnwrap);
-	void OnMeshTriangleHovered(int32 Tid, int32 IndexOfMesh, bool bTidIsFromUnwrap);
-	void UpdateHover();
+	UE_API void OnMeshTriangleClicked(int32 Tid, int32 IndexOfMesh, bool bTidIsFromUnwrap);
+	UE_API void OnMeshTriangleHovered(int32 Tid, int32 IndexOfMesh, bool bTidIsFromUnwrap);
+	UE_API void UpdateHover();
 	/** @param bClearHoverInfo If true, also clears HoverVid, etc in addition to just clearing display. */
-	void ClearHover(bool bClearHoverInfo = true);
+	UE_API void ClearHover(bool bClearHoverInfo = true);
 
-	void ApplyClick();
-	void UpdateToolMessage();
+	UE_API void ApplyClick();
+	UE_API void UpdateToolMessage();
 
 	// Used to remember click info to apply on tick
 	int32 ClickedTid = IndexConstants::InvalidID;
@@ -206,7 +208,9 @@ private:
 
 	UE::Geometry::UVEditorAnalytics::FTargetAnalytics InputTargetAnalytics;
 	FDateTime ToolStartTimeAnalytics;
-	void RecordAnalytics();
+	UE_API void RecordAnalytics();
 
 
 };
+
+#undef UE_API
diff --git a/Engine/Plugins/Editor/UVEditor/Source/UVEditorTools/Public/UVEditorToolAnalyticsUtils.h b/Engine/Plugins/Editor/UVEditor/Source/UVEditorTools/Public/UVEditorToolAnalyticsUtils.h
--- a/Engine/Plugins/Editor/UVEditor/Source/UVEditorTools/Public/UVEditorToolAnalyticsUtils.h
+++ b/Engine/Plugins/Editor/UVEditor/Source/UVEditorTools/Public/UVEditorToolAnalyticsUtils.h
@@ -6,6 +6,8 @@
 #include "EngineAnalytics.h"
 #include "ToolTargets/UVEditorToolMeshInput.h"
 
+#define UE_API UVEDITORTOOLS_API
+
 namespace UE
 {
 namespace Geometry
@@ -24,7 +26,7 @@ FAnalyticsEventAttribute AnalyticsEventAttributeEnum(const FString& AttributeNam
 
 UVEDITORTOOLS_API FString UVEditorAnalyticsEventName(const FString& ToolEventName);
 
-struct UVEDITORTOOLS_API FTargetAnalytics
+struct FTargetAnalytics
 {
 	int64 AllAssetsNumVertices = 0;
 	int64 AllAssetsNumTriangles = 0;
@@ -37,7 +39,7 @@ struct UVEDITORTOOLS_API FTargetAnalytics
 	TArray<int32> PerAssetActiveUVChannelNumVertices;
 	TArray<int32> PerAssetActiveUVChannelNumTriangles;
 
-	void AppendToAttributes(TArray<FAnalyticsEventAttribute>& Attributes, FString Prefix = "") const;
+	UE_API void AppendToAttributes(TArray<FAnalyticsEventAttribute>& Attributes, FString Prefix = "") const;
 };
 
 UVEDITORTOOLS_API FTargetAnalytics CollectTargetAnalytics(const TArray<TObjectPtr<UUVEditorToolMeshInput>>& Targets);
@@ -46,3 +48,5 @@ UVEDITORTOOLS_API FTargetAnalytics CollectTargetAnalytics(const TArray<TObjectPt
 	
 } // end namespace Geometry
 } // end namespace UE
+
+#undef UE_API
diff --git a/Engine/Plugins/Editor/UVEditor/Source/UVEditorTools/Public/UVEditorToolBase.h b/Engine/Plugins/Editor/UVEditor/Source/UVEditorTools/Public/UVEditorToolBase.h
--- a/Engine/Plugins/Editor/UVEditor/Source/UVEditorTools/Public/UVEditorToolBase.h
+++ b/Engine/Plugins/Editor/UVEditor/Source/UVEditorTools/Public/UVEditorToolBase.h
@@ -8,6 +8,8 @@
 
 #include "UVEditorToolBase.generated.h"
 
+#define UE_API UVEDITORTOOLS_API
+
 class UUVEditorToolMeshInput;
 
 /** UObject for IUVEditorGenericBuildableTool */
@@ -33,8 +35,8 @@ public:
  * Simple builder that just instantiates the given class and passes in the targets. Can be used
  * for any UV tools that don't need special handling, as long as they implement IUVEditorGenericBuildableTool.
  */
-UCLASS()
-class UVEDITORTOOLS_API UGenericUVEditorToolBuilder : public UInteractiveToolBuilder
+UCLASS(MinimalAPI)
+class UGenericUVEditorToolBuilder : public UInteractiveToolBuilder
 {
 	GENERATED_BODY()
 public:
@@ -42,10 +44,10 @@ public:
 	/**
 	 * @param ToolClassIn Must implement IUVEditorGenericBuildableTool.
 	 */
-	virtual void Initialize(TArray<TObjectPtr<UUVEditorToolMeshInput>>& TargetsIn, TSubclassOf<UInteractiveTool> ToolClassIn);
+	UE_API virtual void Initialize(TArray<TObjectPtr<UUVEditorToolMeshInput>>& TargetsIn, TSubclassOf<UInteractiveTool> ToolClassIn);
 
-	virtual bool CanBuildTool(const FToolBuilderState& SceneState) const override;
-	virtual UInteractiveTool* BuildTool(const FToolBuilderState& SceneState) const override;
+	UE_API virtual bool CanBuildTool(const FToolBuilderState& SceneState) const override;
+	UE_API virtual UInteractiveTool* BuildTool(const FToolBuilderState& SceneState) const override;
 
 private:
 	// This is a pointer so that it can be updated under the builder without
@@ -53,4 +55,6 @@ private:
 	const TArray<TObjectPtr<UUVEditorToolMeshInput>>* Targets = nullptr;
 
 	TSubclassOf<UInteractiveTool> ToolClass = nullptr;
-};
\ No newline at end of file
+};
+
+#undef UE_API
diff --git a/Engine/Plugins/Editor/UVEditor/Source/UVEditorTools/Public/UVEditorTransformTool.h b/Engine/Plugins/Editor/UVEditor/Source/UVEditorTools/Public/UVEditorTransformTool.h
--- a/Engine/Plugins/Editor/UVEditor/Source/UVEditorTools/Public/UVEditorTransformTool.h
+++ b/Engine/Plugins/Editor/UVEditor/Source/UVEditorTools/Public/UVEditorTransformTool.h
@@ -12,6 +12,8 @@
 
 #include "UVEditorTransformTool.generated.h"
 
+#define UE_API UVEDITORTOOLS_API
+
 class UUVEditorToolMeshInput;
 class UUVEditorUVTransformProperties;
 enum class EUVEditorUVTransformType;
@@ -22,8 +24,8 @@ class UUVEditorTransformTool;
 /**
  * Visualization settings for the TransformTool
  */
-UCLASS()
-class UVEDITORTOOLS_API UUVEditorTransformToolDisplayProperties : public UInteractiveToolPropertySet
+UCLASS(MinimalAPI)
+class UUVEditorTransformToolDisplayProperties : public UInteractiveToolPropertySet
 {
 	GENERATED_BODY()
 public:
@@ -32,55 +34,55 @@ public:
 	bool bDrawPivots = true;
 };
 
-UCLASS()
-class UVEDITORTOOLS_API UUVEditorBaseTransformToolBuilder : public UInteractiveToolBuilder
+UCLASS(MinimalAPI)
+class UUVEditorBaseTransformToolBuilder : public UInteractiveToolBuilder
 {
 	GENERATED_BODY()
 
 public:
-	virtual bool CanBuildTool(const FToolBuilderState& SceneState) const override;
-	virtual UInteractiveTool* BuildTool(const FToolBuilderState& SceneState) const override;
+	UE_API virtual bool CanBuildTool(const FToolBuilderState& SceneState) const override;
+	UE_API virtual UInteractiveTool* BuildTool(const FToolBuilderState& SceneState) const override;
 
 	// This is a pointer so that it can be updated under the builder without
 	// having to set it in the mode after initializing targets.
 	const TArray<TObjectPtr<UUVEditorToolMeshInput>>* Targets = nullptr;
 
 protected:
-	virtual void ConfigureTool(UUVEditorTransformTool* NewTool) const;
+	UE_API virtual void ConfigureTool(UUVEditorTransformTool* NewTool) const;
 };
 
-UCLASS()
-class UVEDITORTOOLS_API UUVEditorTransformToolBuilder : public UUVEditorBaseTransformToolBuilder
+UCLASS(MinimalAPI)
+class UUVEditorTransformToolBuilder : public UUVEditorBaseTransformToolBuilder
 {
 	GENERATED_BODY()
 
 protected:
-	virtual void ConfigureTool(UUVEditorTransformTool* NewTool) const override;
+	UE_API virtual void ConfigureTool(UUVEditorTransformTool* NewTool) const override;
 };
 
-UCLASS()
-class UVEDITORTOOLS_API UUVEditorAlignToolBuilder : public UUVEditorBaseTransformToolBuilder
+UCLASS(MinimalAPI)
+class UUVEditorAlignToolBuilder : public UUVEditorBaseTransformToolBuilder
 {
 	GENERATED_BODY()
 
 protected:
-	virtual void ConfigureTool(UUVEditorTransformTool* NewTool) const  override;
+	UE_API virtual void ConfigureTool(UUVEditorTransformTool* NewTool) const  override;
 };
 
-UCLASS()
-class UVEDITORTOOLS_API UUVEditorDistributeToolBuilder : public UUVEditorBaseTransformToolBuilder
+UCLASS(MinimalAPI)
+class UUVEditorDistributeToolBuilder : public UUVEditorBaseTransformToolBuilder
 {
 	GENERATED_BODY()
 
 protected:
-	virtual void ConfigureTool(UUVEditorTransformTool* NewTool) const  override;
+	UE_API virtual void ConfigureTool(UUVEditorTransformTool* NewTool) const  override;
 };
 
 /**
  * 
  */
-UCLASS()
-class UVEDITORTOOLS_API UUVEditorTransformTool : public UInteractiveTool, public IUVToolSupportsSelection
+UCLASS(MinimalAPI)
+class UUVEditorTransformTool : public UInteractiveTool, public IUVToolSupportsSelection
 {
 	GENERATED_BODY()
 
@@ -94,20 +96,20 @@ public:
 		Targets = TargetsIn;
 	}
 
-	void SetToolMode(const EUVEditorUVTransformType& Mode);
+	UE_API void SetToolMode(const EUVEditorUVTransformType& Mode);
 
-	virtual void Setup() override;
-	virtual void Shutdown(EToolShutdownType ShutdownType) override;
+	UE_API virtual void Setup() override;
+	UE_API virtual void Shutdown(EToolShutdownType ShutdownType) override;
 
-	virtual void OnTick(float DeltaTime) override;
+	UE_API virtual void OnTick(float DeltaTime) override;
 
 	virtual bool HasCancel() const override { return true; }
 	virtual bool HasAccept() const override { return true; }
-	virtual bool CanAccept() const override;
+	UE_API virtual bool CanAccept() const override;
 
-	virtual void DrawHUD(FCanvas* Canvas, IToolsContextRenderAPI* RenderAPI) override;
+	UE_API virtual void DrawHUD(FCanvas* Canvas, IToolsContextRenderAPI* RenderAPI) override;
 
-	virtual void OnPropertyModified(UObject* PropertySet, FProperty* Property) override;
+	UE_API virtual void OnPropertyModified(UObject* PropertySet, FProperty* Property) override;
 
 protected:	
 
@@ -136,5 +138,7 @@ protected:
 
 	UE::Geometry::UVEditorAnalytics::FTargetAnalytics InputTargetAnalytics;
 	FDateTime ToolStartTimeAnalytics;
-	void RecordAnalytics();
+	UE_API void RecordAnalytics();
 };
+
+#undef UE_API
diff --git a/Engine/Plugins/Editor/UVEditor/Source/UVEditorTools/Public/UVEditorUVSnapshotTool.h b/Engine/Plugins/Editor/UVEditor/Source/UVEditorTools/Public/UVEditorUVSnapshotTool.h
--- a/Engine/Plugins/Editor/UVEditor/Source/UVEditorTools/Public/UVEditorUVSnapshotTool.h
+++ b/Engine/Plugins/Editor/UVEditor/Source/UVEditorTools/Public/UVEditorUVSnapshotTool.h
@@ -11,6 +11,8 @@
 
 #include "UVEditorUVSnapshotTool.generated.h"
 
+#define UE_API UVEDITORTOOLS_API
+
 
 // Forward Declarations
 class UUVEditorBakeUVShellProperties;
@@ -18,13 +20,13 @@ class UUVEditorBakeUVShellProperties;
 /**
  * ToolBuilder
  */
-UCLASS()
-class UVEDITORTOOLS_API UUVEditorUVSnapshotToolBuilder : public UInteractiveToolBuilder
+UCLASS(MinimalAPI)
+class UUVEditorUVSnapshotToolBuilder : public UInteractiveToolBuilder
 {
 	GENERATED_BODY()
 public:
-	virtual bool CanBuildTool(const FToolBuilderState& SceneState) const override;
-	virtual UInteractiveTool* BuildTool(const FToolBuilderState& SceneState) const override;
+	UE_API virtual bool CanBuildTool(const FToolBuilderState& SceneState) const override;
+	UE_API virtual UInteractiveTool* BuildTool(const FToolBuilderState& SceneState) const override;
 	// This is a pointer so that it can be updated under the builder without
 	// having to set it in the mode after initializing targets.
 	// Supports only one target
@@ -36,24 +38,24 @@ public:
  *
  * Exports a texture asset of a UV Layout
  */
-UCLASS()
-class UVEDITORTOOLS_API UUVEditorUVSnapshotTool : public UInteractiveTool, public IUVToolSupportsSelection, public UE::Geometry::IGenericDataOperatorFactory<UE::Geometry::FMeshMapBaker>
+UCLASS(MinimalAPI)
+class UUVEditorUVSnapshotTool : public UInteractiveTool, public IUVToolSupportsSelection, public UE::Geometry::IGenericDataOperatorFactory<UE::Geometry::FMeshMapBaker>
 {
 	GENERATED_BODY()
 public:
 	// Begin UInteractiveTool interface
-	virtual void Setup() override;
-	virtual void Shutdown(EToolShutdownType ShutdownType) override;
+	UE_API virtual void Setup() override;
+	UE_API virtual void Shutdown(EToolShutdownType ShutdownType) override;
 	
-	virtual void OnTick(float DeltaTime) override;
+	UE_API virtual void OnTick(float DeltaTime) override;
 	virtual void Render(IToolsContextRenderAPI* RenderAPI) override { };
 	virtual bool HasCancel() const override { return true; }
 	virtual bool HasAccept() const override { return true; }
-	virtual bool CanAccept() const override;
+	UE_API virtual bool CanAccept() const override;
 	// End UInteractiveTool interface
 
 	// Begin IGenericDataOperatorFactory interface
-	virtual TUniquePtr<UE::Geometry::TGenericDataOperator<UE::Geometry::FMeshMapBaker>> MakeNewOperator() override;
+	UE_API virtual TUniquePtr<UE::Geometry::TGenericDataOperator<UE::Geometry::FMeshMapBaker>> MakeNewOperator() override;
 	// End IGenericDataOperatorFactory interface
 	
 	/**
@@ -96,14 +98,14 @@ private:
 	 * 
 	 * @param NewResult the resulting FMeshMapBaker from the background compute
 	 */
-	void OnMapUpdated(const TUniquePtr<UE::Geometry::FMeshMapBaker>& NewResult);
+	UE_API void OnMapUpdated(const TUniquePtr<UE::Geometry::FMeshMapBaker>& NewResult);
 
 	/**
 	 * Updates the preview material on the preview quad with the
 	 * computed results. Invoked by OnMapsUpdated.
 	 * Also sets the result to what is currently in CachedUVMap
 	 */
-	void UpdateVisualization();
+	UE_API void UpdateVisualization();
 
 	/**
 	 * Internal cache of bake uv texture result
@@ -114,27 +116,27 @@ private:
 	/**
 	 * Uses Preview Geometry to draw a preview of the bake in the unwrap viewport
 	 */
-	void SetUpPreviewQuad();
+	UE_API void SetUpPreviewQuad();
 
 	/**
 	 * Retrieves the Material for the preview quad, or the preview result
 	 */
-	UMaterialInstanceDynamic* GetMaterialForQuad();
+	UE_API UMaterialInstanceDynamic* GetMaterialForQuad();
 
 	/**
 	 * Create texture asset from our result Texture2D
 	 * @param Texture the result texture to create
 	 */
-	void CreateTextureAsset(const TObjectPtr<UTexture2D>& Texture) const;
+	UE_API void CreateTextureAsset(const TObjectPtr<UTexture2D>& Texture) const;
 
 	/**
 	 * Initialize the list of all UV Layer names
 	 */
-	static void InitializeUVLayerNames(TArray<FString>& UVLayerNamesList, const int16 NumUVLayers);
+	static UE_API void InitializeUVLayerNames(TArray<FString>& UVLayerNamesList, const int16 NumUVLayers);
 };
 
-UCLASS()
-class UVEDITORTOOLS_API UUVEditorBakeUVShellProperties : public UInteractiveToolPropertySet
+UCLASS(MinimalAPI)
+class UUVEditorBakeUVShellProperties : public UInteractiveToolPropertySet
 {
 	GENERATED_BODY()
 public:
@@ -182,4 +184,6 @@ public:
 
 	UPROPERTY(meta = (TransientToolProperty))
 	TArray<FString> TargetUVLayerNamesList;
-};
\ No newline at end of file
+};
+
+#undef UE_API
diff --git a/Engine/Plugins/Editor/UVEditor/Source/UVEditorTools/Public/UVEditorUXSettings.h b/Engine/Plugins/Editor/UVEditor/Source/UVEditorTools/Public/UVEditorUXSettings.h
--- a/Engine/Plugins/Editor/UVEditor/Source/UVEditorTools/Public/UVEditorUXSettings.h
+++ b/Engine/Plugins/Editor/UVEditor/Source/UVEditorTools/Public/UVEditorUXSettings.h
@@ -5,7 +5,9 @@
 
 #include "CoreGlobals.h"
 #include "HAL/IConsoleManager.h"
-class UVEDITORTOOLS_API FUVEditorUXSettings
+
+#define UE_API UVEDITORTOOLS_API
+class FUVEditorUXSettings
 {
 public:
 	// The following are UV Editor specific style items that don't, strictly, matter to the SlateStyleSet,
@@ -15,110 +17,110 @@ public:
 
 	// The ScaleFactor just scales the mesh up. Scaling the mesh up makes it easier to zoom in
 	// further into the display before getting issues with the camera near plane distance.
-	static const float UVMeshScalingFactor;
+	static UE_API const float UVMeshScalingFactor;
 
 	// Position to place the 2D camera far plane relative to world z
-	static const float CameraFarPlaneWorldZ;
+	static UE_API const float CameraFarPlaneWorldZ;
 
 	// The near plane gets positioned some proportion to z = 0. We don't use a constant value because our depth offset values are percentage-based
 	// Lower proportion move the plane nearer to world z
 	// Note: This serves as an upper bound for all other depth offsets - higher than this value risks being clipped
-	static const float CameraNearPlaneProportionZ;;
+	static UE_API const float CameraNearPlaneProportionZ;;
 
 
 	// 2D Unwrap Display Properties
-	static const float UnwrapTriangleOpacity;
-	static const float UnwrapTriangleDepthOffset;
-	static const float UnwrapTriangleOpacityWithBackground;
+	static UE_API const float UnwrapTriangleOpacity;
+	static UE_API const float UnwrapTriangleDepthOffset;
+	static UE_API const float UnwrapTriangleOpacityWithBackground;
 
-	static const float WireframeDepthOffset;
-	static const float UnwrapBoundaryHueShift;
-	static const float UnwrapBoundarySaturation;
-	static const float UnwrapBoundaryValue;
-	static const FColor UnwrapTriangleFillColor;
-	static const FColor UnwrapTriangleWireframeColor;
+	static UE_API const float WireframeDepthOffset;
+	static UE_API const float UnwrapBoundaryHueShift;
+	static UE_API const float UnwrapBoundarySaturation;
+	static UE_API const float UnwrapBoundaryValue;
+	static UE_API const FColor UnwrapTriangleFillColor;
+	static UE_API const FColor UnwrapTriangleWireframeColor;
 
-	static FLinearColor GetTriangleColorByTargetIndex(int32 TargetIndex);
-	static FLinearColor GetWireframeColorByTargetIndex(int32 TargetIndex);
-	static FLinearColor GetBoundaryColorByTargetIndex(int32 TargetIndex);
+	static UE_API FLinearColor GetTriangleColorByTargetIndex(int32 TargetIndex);
+	static UE_API FLinearColor GetWireframeColorByTargetIndex(int32 TargetIndex);
+	static UE_API FLinearColor GetBoundaryColorByTargetIndex(int32 TargetIndex);
 
 	// Provides a color blindness friendly ramp for visualizing metrics and other data
 	// Domain is [0 , 1]
-	static FColor MakeCividisColorFromScalar(float Scalar);
+	static UE_API FColor MakeCividisColorFromScalar(float Scalar);
 
 	// Provides a color blindness friendly ramp for visualizing diverging metrics and other data
 	// Domain is [-0.5, 0.5] 
-	static FColor MakeTurboColorFromScalar(float Scalar);
+	static UE_API FColor MakeTurboColorFromScalar(float Scalar);
 
 	// Wireframe Properties
-	static const float WireframeThickness;
-	static const float BoundaryEdgeThickness;
+	static UE_API const float WireframeThickness;
+	static UE_API const float BoundaryEdgeThickness;
 
 	// Selection Highlighting Properties
-	static const float SelectionTriangleOpacity;
-	static const FColor SelectionTriangleFillColor;
-	static const FColor SelectionTriangleWireframeColor;
+	static UE_API const float SelectionTriangleOpacity;
+	static UE_API const FColor SelectionTriangleFillColor;
+	static UE_API const FColor SelectionTriangleWireframeColor;
 
-	static const float SelectionHoverTriangleOpacity;
-	static const FColor SelectionHoverTriangleFillColor;
-	static const FColor SelectionHoverTriangleWireframeColor;
+	static UE_API const float SelectionHoverTriangleOpacity;
+	static UE_API const FColor SelectionHoverTriangleFillColor;
+	static UE_API const FColor SelectionHoverTriangleWireframeColor;
 
-	static const float LivePreviewHighlightThickness;
-	static const float LivePreviewHighlightPointSize;
-	static const float LivePreviewHighlightDepthOffset;
+	static UE_API const float LivePreviewHighlightThickness;
+	static UE_API const float LivePreviewHighlightPointSize;
+	static UE_API const float LivePreviewHighlightDepthOffset;
 
-	static const float SelectionLineThickness;
-	static const float SelectionPointThickness;
-	static const float SelectionWireframeDepthBias;
-	static const float SelectionTriangleDepthBias;
-	static const float SelectionHoverWireframeDepthBias;
-	static const float SelectionHoverTriangleDepthBias;
+	static UE_API const float SelectionLineThickness;
+	static UE_API const float SelectionPointThickness;
+	static UE_API const float SelectionWireframeDepthBias;
+	static UE_API const float SelectionTriangleDepthBias;
+	static UE_API const float SelectionHoverWireframeDepthBias;
+	static UE_API const float SelectionHoverTriangleDepthBias;
 
-	static const FColor LivePreviewExistingSeamColor;
-	static const float LivePreviewExistingSeamThickness;
-	static const float LivePreviewExistingSeamDepthBias;
+	static UE_API const FColor LivePreviewExistingSeamColor;
+	static UE_API const float LivePreviewExistingSeamThickness;
+	static UE_API const float LivePreviewExistingSeamDepthBias;
 
 	// These are currently used by the seam tool but can be generally used by
 	// tools for displaying paths.
-	static const FColor ToolLockedCutPathColor;
-	static const FColor ToolLockedJoinPathColor;
-	static const float ToolLockedPathThickness;
-	static const float ToolLockedPathDepthBias;
-	static const FColor ToolExtendCutPathColor;
-	static const FColor ToolExtendJoinPathColor;
-	static const float ToolExtendPathThickness;
-	static const float ToolExtendPathDepthBias;
-	static const FColor ToolCompletionPathColor;
-
-	static const float ToolPointSize;
+	static UE_API const FColor ToolLockedCutPathColor;
+	static UE_API const FColor ToolLockedJoinPathColor;
+	static UE_API const float ToolLockedPathThickness;
+	static UE_API const float ToolLockedPathDepthBias;
+	static UE_API const FColor ToolExtendCutPathColor;
+	static UE_API const FColor ToolExtendJoinPathColor;
+	static UE_API const float ToolExtendPathThickness;
+	static UE_API const float ToolExtendPathDepthBias;
+	static UE_API const FColor ToolCompletionPathColor;
+
+	static UE_API const float ToolPointSize;
 
 	// Sew Action styling
-	static const float SewLineHighlightThickness;
-	static const float SewLineDepthOffset;
-	static const FColor SewSideLeftColor;
-	static const FColor SewSideRightColor;
+	static UE_API const float SewLineHighlightThickness;
+	static UE_API const float SewLineDepthOffset;
+	static UE_API const FColor SewSideLeftColor;
+	static UE_API const FColor SewSideRightColor;
 
 	// Grid 
-	static const float AxisThickness;
-	static const float GridMajorThickness;
-	static const FColor XAxisColor;
-	static const FColor YAxisColor;
-	static const FColor RulerXColor;
-	static const FColor RulerYColor;
-	static const FColor GridMajorColor;
-	static const FColor GridMinorColor;
-	static const int32 GridSubdivisionsPerLevel;
-	static const int32 GridLevels;
-	static const int32 RulerSubdivisionLevel;
+	static UE_API const float AxisThickness;
+	static UE_API const float GridMajorThickness;
+	static UE_API const FColor XAxisColor;
+	static UE_API const FColor YAxisColor;
+	static UE_API const FColor RulerXColor;
+	static UE_API const FColor RulerYColor;
+	static UE_API const FColor GridMajorColor;
+	static UE_API const FColor GridMinorColor;
+	static UE_API const int32 GridSubdivisionsPerLevel;
+	static UE_API const int32 GridLevels;
+	static UE_API const int32 RulerSubdivisionLevel;
 
 	// Pivots
-	static const int32 PivotCircleNumSides;
-	static const float PivotCircleRadius;
-	static const float PivotLineThickness;
-	static const FColor PivotLineColor;
+	static UE_API const int32 PivotCircleNumSides;
+	static UE_API const float PivotCircleRadius;
+	static UE_API const float PivotLineThickness;
+	static UE_API const FColor PivotLineColor;
 
 	// Background
-	static const float BackgroundQuadDepthOffset;
+	static UE_API const float BackgroundQuadDepthOffset;
 
 	//---------------------------
 	//  Common Utility Methods
@@ -136,7 +138,7 @@ public:
 	 *
 	 * Like other world-UV conversion functions, clients should not know about the details of this conversion.
 	 */
-	static FVector3d ExternalUVToUnwrapWorldPosition(const FVector2f& UV);
+	static UE_API FVector3d ExternalUVToUnwrapWorldPosition(const FVector2f& UV);
 
 	/**
 	 * Converts from point in unwrap world space, to UV value as displayed to user or seen before import. This
@@ -145,17 +147,17 @@ public:
 	 *
 	 * Like other world-UV conversion functions, clients should not know about the details of this conversion.
 	 */
-	static FVector2f UnwrapWorldPositionToExternalUV(const FVector3d& VertPosition);
+	static UE_API FVector2f UnwrapWorldPositionToExternalUV(const FVector3d& VertPosition);
 
 	/**
 	 * Converts from UV value as displayed to user or seen before import, to UV as stored on a mesh in Unreal.
 	 */
-	static FVector2f ExternalUVToInternalUV(const FVector2f& UV);
+	static UE_API FVector2f ExternalUVToInternalUV(const FVector2f& UV);
 
 	/**
 	 * Converts from UV as stored on a mesh in Unreal to UV as displayed to user or seen in an external program.
 	 */
-	static FVector2f InternalUVToExternalUV(const FVector2f& UV);
+	static UE_API FVector2f InternalUVToExternalUV(const FVector2f& UV);
 
 	/**
 	 * Convert from UV as stored on a mesh in Unreal to world position in the UV editor unwrap world. This allows
@@ -163,7 +165,7 @@ public:
 	 * 
 	 * Like other world-UV conversion functions, clients should not know about the details of this conversion.
 	 */
-	static FVector3d UVToVertPosition(const FVector2f& UV);
+	static UE_API FVector3d UVToVertPosition(const FVector2f& UV);
 
 	/**
 	 * Converts from position in UV editor unwrap world to UV value as stored on a mesh in Unreal. This allows
@@ -171,23 +173,25 @@ public:
 	 * 
 	 * Like other world-UV conversion functions, clients should not know about the details of this conversion.
 	 */
-	static FVector2f VertPositionToUV(const FVector3d& VertPosition);
+	static UE_API FVector2f VertPositionToUV(const FVector3d& VertPosition);
 
 
 	//--------------------------------
 	// CVARs for Experimental Features
 	//--------------------------------
 
-	static TAutoConsoleVariable<int32> CVarEnablePrototypeUDIMSupport;
+	static UE_API TAutoConsoleVariable<int32> CVarEnablePrototypeUDIMSupport;
 
 	//--------------------------------
 	// Values for Snapping
 	//--------------------------------
-	static float LocationSnapValue(int32 LocationSnapMenuIndex);
-	static int32 MaxLocationSnapValue();
+	static UE_API float LocationSnapValue(int32 LocationSnapMenuIndex);
+	static UE_API int32 MaxLocationSnapValue();
 
 private:
 
-	FUVEditorUXSettings();
-	~FUVEditorUXSettings();
-};
\ No newline at end of file
+	UE_API FUVEditorUXSettings();
+	UE_API ~FUVEditorUXSettings();
+};
+
+#undef UE_API
diff --git a/Engine/Plugins/Editor/UVEditor/Source/UVEditorTools/Public/UVSelectTool.h b/Engine/Plugins/Editor/UVEditor/Source/UVEditorTools/Public/UVSelectTool.h
--- a/Engine/Plugins/Editor/UVEditor/Source/UVEditorTools/Public/UVSelectTool.h
+++ b/Engine/Plugins/Editor/UVEditor/Source/UVEditorTools/Public/UVSelectTool.h
@@ -14,20 +14,22 @@
 
 #include "UVSelectTool.generated.h"
 
+#define UE_API UVEDITORTOOLS_API
+
 class UCombinedTransformGizmo;
 class UTransformProxy;
 class UUVEditorToolMeshInput;
 class UUVToolEmitChangeAPI;
 class UUVToolViewportButtonsAPI;
 
-UCLASS()
-class UVEDITORTOOLS_API UUVSelectToolBuilder : public UInteractiveToolBuilder
+UCLASS(MinimalAPI)
+class UUVSelectToolBuilder : public UInteractiveToolBuilder
 {
 	GENERATED_BODY()
 public:
 
-	virtual bool CanBuildTool(const FToolBuilderState& SceneState) const override;
-	virtual UInteractiveTool* BuildTool(const FToolBuilderState& SceneState) const override;
+	UE_API virtual bool CanBuildTool(const FToolBuilderState& SceneState) const override;
+	UE_API virtual UInteractiveTool* BuildTool(const FToolBuilderState& SceneState) const override;
 
 	// This is a pointer so that it can be updated under the builder without
 	// having to set it in the mode after initializing targets.
@@ -39,8 +41,8 @@ public:
  * selection api to allow the user to select elements, and has a gizmo that can be used to
  * transform these elements.
  */
-UCLASS()
-class UVEDITORTOOLS_API UUVSelectTool : public UInteractiveTool, public IInteractiveToolNestedAcceptCancelAPI, public IUVToolSupportsSelection
+UCLASS(MinimalAPI)
+class UUVSelectTool : public UInteractiveTool, public IInteractiveToolNestedAcceptCancelAPI, public IUVToolSupportsSelection
 {
 	GENERATED_BODY()
 
@@ -54,37 +56,37 @@ public:
 		Targets = TargetsIn;
 	}
 
-	void SelectAll();
+	UE_API void SelectAll();
 
 	// Used by undo/redo.
-	FTransform GetGizmoTransform() const;
-	void SetGizmoTransform(const FTransform& NewTransform);
+	UE_API FTransform GetGizmoTransform() const;
+	UE_API void SetGizmoTransform(const FTransform& NewTransform);
 
 	// UInteractiveTool
-	virtual void Setup() override;
-	virtual void Shutdown(EToolShutdownType ShutdownType) override;
-	virtual void OnPropertyModified(UObject* PropertySet, FProperty* Property) override;
-	virtual void Render(IToolsContextRenderAPI* RenderAPI) override;
-	virtual void DrawHUD(FCanvas* Canvas, IToolsContextRenderAPI* RenderAPI) override;
-	virtual void OnTick(float DeltaTime) override;
+	UE_API virtual void Setup() override;
+	UE_API virtual void Shutdown(EToolShutdownType ShutdownType) override;
+	UE_API virtual void OnPropertyModified(UObject* PropertySet, FProperty* Property) override;
+	UE_API virtual void Render(IToolsContextRenderAPI* RenderAPI) override;
+	UE_API virtual void DrawHUD(FCanvas* Canvas, IToolsContextRenderAPI* RenderAPI) override;
+	UE_API virtual void OnTick(float DeltaTime) override;
 	virtual bool HasCancel() const override { return false; }
 	virtual bool HasAccept() const override { return false; }
 
 	// IInteractiveToolNestedAcceptCancelAPI
 	virtual bool SupportsNestedCancelCommand() override { return true; }
-	virtual bool CanCurrentlyNestedCancel() override;
-	virtual bool ExecuteNestedCancelCommand() override;
+	UE_API virtual bool CanCurrentlyNestedCancel() override;
+	UE_API virtual bool ExecuteNestedCancelCommand() override;
 
 protected:
-	virtual void OnSelectionChanged(bool bEmitChangeAllowed, uint32 SelectionChangeType);
+	UE_API virtual void OnSelectionChanged(bool bEmitChangeAllowed, uint32 SelectionChangeType);
 
 	// Callbacks we'll receive from the gizmo proxy
-	virtual void GizmoTransformChanged(UTransformProxy* Proxy, FTransform Transform);
-	virtual void GizmoTransformStarted(UTransformProxy* Proxy);
-	virtual void GizmoTransformEnded(UTransformProxy* Proxy);
+	UE_API virtual void GizmoTransformChanged(UTransformProxy* Proxy, FTransform Transform);
+	UE_API virtual void GizmoTransformStarted(UTransformProxy* Proxy);
+	UE_API virtual void GizmoTransformEnded(UTransformProxy* Proxy);
 
-	virtual void ApplyGizmoTransform();
-	virtual void UpdateGizmo(bool bForceRecomputeSelectionCenters = false);
+	UE_API virtual void ApplyGizmoTransform();
+	UE_API virtual void UpdateGizmo(bool bForceRecomputeSelectionCenters = false);
 
 	UPROPERTY()
 	TArray<TObjectPtr<UUVEditorToolMeshInput>> Targets;
@@ -118,6 +120,7 @@ protected:
 	TArray<TArray<FVector3d>> MovingVertOriginalPositionsPerSelection;
 
 private:
-	void ReinitializeFromSelection(bool bForceRecomputeSelectionCenters);
+	UE_API void ReinitializeFromSelection(bool bForceRecomputeSelectionCenters);
 };
 
+#undef UE_API
diff --git a/Engine/Plugins/Editor/UVEditor/Source/UVEditorToolsEditorOnly/Public/UVEditorParameterizeMeshTool.h b/Engine/Plugins/Editor/UVEditor/Source/UVEditorToolsEditorOnly/Public/UVEditorParameterizeMeshTool.h
--- a/Engine/Plugins/Editor/UVEditor/Source/UVEditorToolsEditorOnly/Public/UVEditorParameterizeMeshTool.h
+++ b/Engine/Plugins/Editor/UVEditor/Source/UVEditorToolsEditorOnly/Public/UVEditorParameterizeMeshTool.h
@@ -16,6 +16,8 @@
 
 #include "UVEditorParameterizeMeshTool.generated.h"
 
+#define UE_API UVEDITORTOOLSEDITORONLY_API
+
 // predeclarations
 class UDynamicMeshComponent;
 class UMaterialInterface;
@@ -28,13 +30,13 @@ class UParameterizeMeshToolPatchBuilderProperties;
 class UParameterizeMeshOperatorFactory;
 class UPolygroupLayersProperties;
 
-UCLASS()
-class UVEDITORTOOLSEDITORONLY_API UUVEditorParameterizeMeshToolBuilder : public UInteractiveToolBuilder
+UCLASS(MinimalAPI)
+class UUVEditorParameterizeMeshToolBuilder : public UInteractiveToolBuilder
 {
 	GENERATED_BODY()
 public:
-	virtual bool CanBuildTool(const FToolBuilderState& SceneState) const override;
-	virtual UInteractiveTool* BuildTool(const FToolBuilderState& SceneState) const override;
+	UE_API virtual bool CanBuildTool(const FToolBuilderState& SceneState) const override;
+	UE_API virtual UInteractiveTool* BuildTool(const FToolBuilderState& SceneState) const override;
 
 	// This is a pointer so that it can be updated under the builder without
 	// having to set it in the mode after initializing targets.
@@ -46,8 +48,8 @@ public:
  * UParameterizeMeshTool automatically decomposes the input mesh into charts, solves for UVs,
  * and then packs the resulting charts
  */
-UCLASS()
-class UVEDITORTOOLSEDITORONLY_API UUVEditorParameterizeMeshTool : public UInteractiveTool
+UCLASS(MinimalAPI)
+class UUVEditorParameterizeMeshTool : public UInteractiveTool
 {
 	GENERATED_BODY()
 
@@ -61,16 +63,16 @@ public:
 	}
 
 
-	virtual void Setup() override;
-	virtual void Shutdown(EToolShutdownType ShutdownType) override;
+	UE_API virtual void Setup() override;
+	UE_API virtual void Shutdown(EToolShutdownType ShutdownType) override;
 
-	virtual void OnTick(float DeltaTime) override;
+	UE_API virtual void OnTick(float DeltaTime) override;
 
 	virtual bool HasCancel() const override { return true; }
 	virtual bool HasAccept() const override { return true; }
-	virtual bool CanAccept() const override;
+	UE_API virtual bool CanAccept() const override;
 
-	virtual void OnPropertyModified(UObject* PropertySet, FProperty* Property) override;
+	UE_API virtual void OnPropertyModified(UObject* PropertySet, FProperty* Property) override;
 
 protected:
 	UPROPERTY()
@@ -94,7 +96,7 @@ protected:
 	UPROPERTY()
 	TArray<TObjectPtr<UParameterizeMeshOperatorFactory>> Factories;
 
-	void OnMethodTypeChanged();
+	UE_API void OnMethodTypeChanged();
 
 	//
 	// Analytics
@@ -102,6 +104,7 @@ protected:
 	
 	UE::Geometry::UVEditorAnalytics::FTargetAnalytics InputTargetAnalytics;
 	FDateTime ToolStartTimeAnalytics;
-	void RecordAnalytics();
+	UE_API void RecordAnalytics();
 };
 
+#undef UE_API
