UVEditor: Fix bug where if a collection of objects in the viewport are selected, where at least two objects are backed by the same asset, the UV Editor would fail to position them relative to their viewport locations.

diff --git a/Engine/Plugins/Editor/UVEditor/Source/UVEditor/Private/UVEditorMode.cpp b/Engine/Plugins/Editor/UVEditor/Source/UVEditor/Private/UVEditorMode.cpp
--- a/Engine/Plugins/Editor/UVEditor/Source/UVEditor/Private/UVEditorMode.cpp
+++ b/Engine/Plugins/Editor/UVEditor/Source/UVEditor/Private/UVEditorMode.cpp
@@ -1232,9 +1232,19 @@ void UUVEditorMode::InitializeTargets()
 	// InitializeModeContexts needs to have been called first so that we have the 3d preview world ready.
 	check(LivePreviewWorld);
 
+	// Nathan - Todo: This is the start of a larger fix for the overall problem of what
+	// to do if we're loading more than one object into the UV Editor that's backed by
+	// the same asset. Ideally we should show as many copies of this object in the 3D preview
+	// pane as we have... but this is a big change.
+	// 
+	// So for now, we are at least keeping track of, at this stage, which targets map to
+	// which object(s) and providing correct transforms for at least one of them.
+	typedef TArray<int32> FObjectToTargets;
+	FObjectToTargets ObjectToTargetMap;
+
 	// Build the tool targets that provide us with 3d dynamic meshes
 	UUVEditorSubsystem* UVSubsystem = GEditor->GetEditorSubsystem<UUVEditorSubsystem>();
-	UVSubsystem->BuildTargets(OriginalObjectsToEdit, GetToolTargetRequirements(), ToolTargets);
+	UVSubsystem->BuildTargets(OriginalObjectsToEdit, GetToolTargetRequirements(), ToolTargets, &ObjectToTargetMap);
 
 	// Collect the 3d dynamic meshes from targets. There will always be one for each asset, and the AssetID
 	// of each asset will be the index into these arrays. Individual input objects (representing a specific
@@ -1290,6 +1300,14 @@ void UUVEditorMode::InitializeTargets()
 		{
 			ToolInputObject->AppliedPreview->PreviewMesh->SetTransform(Transforms[AssetID]);
 		}
+		else
+		{			
+			int32 TargetIndex = ObjectToTargetMap.IndexOfByKey(AssetID);
+			if (ensure(TargetIndex != -1))
+			{
+				ToolInputObject->AppliedPreview->PreviewMesh->SetTransform(Transforms[TargetIndex]);
+			}
+		}
 
 		UMaterialInstanceDynamic* MatInstance = UMaterialInstanceDynamic::Create(Material, GetToolManager());
 		MatInstance->SetVectorParameterValue(TEXT("Color"), FUVEditorUXSettings::GetTriangleColorByTargetIndex(AssetID));
diff --git a/Engine/Plugins/Editor/UVEditor/Source/UVEditor/Private/UVEditorSubsystem.cpp b/Engine/Plugins/Editor/UVEditor/Source/UVEditor/Private/UVEditorSubsystem.cpp
--- a/Engine/Plugins/Editor/UVEditor/Source/UVEditor/Private/UVEditorSubsystem.cpp
+++ b/Engine/Plugins/Editor/UVEditor/Source/UVEditor/Private/UVEditorSubsystem.cpp
@@ -200,14 +200,22 @@ bool UUVEditorSubsystem::AreAssetsValidTargets(const TArray<FAssetData>& InAsset
 }
 
 void UUVEditorSubsystem::BuildTargets(const TArray<TObjectPtr<UObject>>& ObjectsIn, 
-	const FToolTargetTypeRequirements& TargetRequirements, TArray<TObjectPtr<UToolTarget>>& TargetsOut)
+	const FToolTargetTypeRequirements& TargetRequirements, TArray<TObjectPtr<UToolTarget>>& TargetsOut,
+	TArray<int32>* ObjectToTargetMapOut)
 {
 	TargetsOut.Reset();
+	if (ObjectToTargetMapOut)
+	{
+		ObjectToTargetMapOut->SetNum(ObjectsIn.Num());
+	}
 
-	TSet<UObject*> TargetAssets;
+	// Map from Object to Index into Target Array
+	TMap<UObject*, uint32> TargetAssets;
 
-	for (UObject* Object : ObjectsIn)
+	for (int32 ObjectIndex = 0; ObjectIndex < ObjectsIn.Num(); ++ObjectIndex)
 	{
+		UObject* Object = ObjectsIn[ObjectIndex];
+
 		if (UToolTarget* Target = ToolTargetManager->BuildTarget(Object, TargetRequirements))
 		{
 			// If this is an asset backed target, only add it if the asset is not already added to the list.
@@ -216,13 +224,31 @@ void UUVEditorSubsystem::BuildTargets(const TArray<TObjectPtr<UObject>>& Objects
 				UObject* SourceData = AssetTarget->GetSourceData();
 				if (!TargetAssets.Contains(SourceData))
 				{
-					TargetAssets.Add(SourceData);
 					TargetsOut.Add(Target);
+					TargetAssets.Add(SourceData, TargetsOut.Num()-1);
+				}
+
+				if (ObjectToTargetMapOut)
+				{
+					(*ObjectToTargetMapOut)[ObjectIndex] = TargetAssets[SourceData];
 				}
 			}
 			else
 			{
 				TargetsOut.Add(Target);
+				if (ObjectToTargetMapOut)
+				{
+					(*ObjectToTargetMapOut)[ObjectIndex] = (TargetsOut.Num() - 1);
+				}
+			}
+
+
+		}
+		else
+		{
+			if (ObjectToTargetMapOut)
+			{
+				(*ObjectToTargetMapOut)[ObjectIndex] = -1;
 			}
 		}
 	}
diff --git a/Engine/Plugins/Editor/UVEditor/Source/UVEditor/Public/UVEditorSubsystem.h b/Engine/Plugins/Editor/UVEditor/Source/UVEditor/Public/UVEditorSubsystem.h
--- a/Engine/Plugins/Editor/UVEditor/Source/UVEditor/Public/UVEditorSubsystem.h
+++ b/Engine/Plugins/Editor/UVEditor/Source/UVEditor/Public/UVEditorSubsystem.h
@@ -55,10 +55,15 @@ public:
 
 	/** 
 	 * Tries to build the core targets that provide meshes for UV tools to work on.
+	 * 
+	 * ObjectToTargetMapOut provides a mapping from the input Object set to the output Target set,
+	 * since these might not be 1:1 and we may need to know which objects are represented by a single
+	 * source target
 	 */
 	UE_API virtual void BuildTargets(const TArray<TObjectPtr<UObject>>& ObjectsIn, 
 		const FToolTargetTypeRequirements& TargetRequirements,
-		TArray<TObjectPtr<UToolTarget>>& TargetsOut);
+		TArray<TObjectPtr<UToolTarget>>& TargetsOut,
+		TArray<int32>* ObjectToTargetMapOut = nullptr);
 
 	/**
 	 * Either brings to the front an existing UV editor instance that is editing one of
